ElectraDecoders: Added support for missing frames in H.264 in the D3D12 video decoder.

#rnx

diff --git a/Engine/Plugins/Media/D3D12VideoDecodersElectra/Source/Private/VideoDecoder_D3D12_H264.cpp b/Engine/Plugins/Media/D3D12VideoDecodersElectra/Source/Private/VideoDecoder_D3D12_H264.cpp
--- a/Engine/Plugins/Media/D3D12VideoDecodersElectra/Source/Private/VideoDecoder_D3D12_H264.cpp
+++ b/Engine/Plugins/Media/D3D12VideoDecodersElectra/Source/Private/VideoDecoder_D3D12_H264.cpp
@@ -222,11 +222,18 @@ IElectraDecoder::EDecoderError FD3D12VideoDecoder_H264::DecodeAccessUnit(const F
 			if (!CurrentConfig.VideoDecoderHeap.IsValid())
 			{
 				int32 DPBSize = spsPtr->GetDPBSize();
-				int32 NumFrames = spsPtr->GetDPBSize() + 1;	// 1 extra for the current frame that's not in the DPB yet.
+				int32 NumFrames = spsPtr->GetDPBSize() + 2;	// 1 extra for the current frame that's not in the DPB yet, and 1 extra that acts as a 'missing' frame.
 				if (!CreateDecoderHeapAndDPB(DPBSize, NumFrames, 16))
 				{
 					return IElectraDecoder::EDecoderError::Error;
 				}
+				MissingReferenceFrame = DPB->GetNextUnusedFrame();
+				if (!MissingReferenceFrame.IsValid())
+				{
+					PostError(0, TEXT("Could not create empty frame used to fill in for missing frames"), ERRCODE_INTERNAL_FAILED_TO_DECODE);
+					return IElectraDecoder::EDecoderError::Error;
+				}
+
 			}
 
 		}
@@ -286,11 +293,6 @@ IElectraDecoder::EDecoderError FD3D12VideoDecoder_H264::DecodeSlicesH264(const F
 	}
 
 	// Some capability checks
-	if (InSequenceParameterSet.gaps_in_frame_num_value_allowed_flag)
-	{
-		PostError(0, FString::Printf(TEXT("DecodeSlicesH264() failed. Cannot decode streams that have gaps_in_frame_num_value_allowed_flag set")), ERRCODE_INTERNAL_FAILED_TO_DECODE);
-		return IElectraDecoder::EDecoderError::Error;
-	}
 	if (InSequenceParameterSet.mb_adaptive_frame_field_flag || InSliceInfos[0].Header.field_pic_flag || InSliceInfos[0].Header.bottom_field_flag)
 	{
 		PostError(0, FString::Printf(TEXT("DecodeSlicesH264() failed. Cannot decode interlaced video.")), ERRCODE_INTERNAL_FAILED_TO_DECODE);
@@ -457,6 +459,18 @@ IElectraDecoder::EDecoderError FD3D12VideoDecoder_H264::DecodeSlicesH264(const F
 		PostError(0, FString::Printf(TEXT("DecodeSlicesH264() failed. %s"), *BitstreamParamsH264.DPBPOC.GetLastError()), ERRCODE_INTERNAL_FAILED_TO_DECODE);
 		return IElectraDecoder::EDecoderError::Error;
 	}
+	TArray<ElectraDecodersUtil::MPEG::H264::FOutputFrameInfo> OutputFrameInfos, UnrefFrameInfos;
+	// Handle potentially missing frames. If there are any, an entry must be made in the DPB which could result in
+	// output of one or many already decoded frames that we need to handle first.
+	BitstreamParamsH264.DPBPOC.HandleMissingFrames(OutputFrameInfos, UnrefFrameInfos, InSliceInfos[0].NalUnitType, InSliceInfos[0].NalRefIdc, InSliceInfos[0].Header, InSequenceParameterSet);
+	IElectraDecoder::EDecoderError MissingFrameOutputResult = HandleOutputListH264(OutputFrameInfos);
+	OutputFrameInfos.Empty();
+	UnrefFrameInfos.Empty();
+	if (MissingFrameOutputResult != IElectraDecoder::EDecoderError::None)
+	{
+		return MissingFrameOutputResult;
+	}
+
 	// Update the current POC values.
 	if (!BitstreamParamsH264.DPBPOC.UpdatePOC(InSliceInfos[0].NalUnitType, InSliceInfos[0].NalRefIdc, InSliceInfos[0].Header, InSequenceParameterSet))
 	{
@@ -473,11 +487,18 @@ IElectraDecoder::EDecoderError FD3D12VideoDecoder_H264::DecodeSlicesH264(const F
 		pp.RefFrameList[i].bPicEntry = 0xff;
 		if (i < RefFrames.Num())
 		{
-			TSharedPtr<FDecodedFrame, ESPMode::ThreadSafe> refFrame = DPB->GetFrameAtIndex(RefFrames[i].UserFrameInfo.IndexInBuffer);
+			TSharedPtr<FDecodedFrame, ESPMode::ThreadSafe> refFrame;
+			if (RefFrames[i].UserFrameInfo.IndexInBuffer >= 0)
+			{
+				refFrame = DPB->GetFrameAtIndex(RefFrames[i].UserFrameInfo.IndexInBuffer);
+			}
+			else
+			{
+				refFrame = MissingReferenceFrame;
+			}
 			if (refFrame.IsValid())
 			{
 				int32 dpbPos = refFrame->IndexInPictureBuffer;
-				check(RefFrames[i].UserFrameInfo.IndexInBuffer == dpbPos);
 				fdr->ReferenceFrameList[dpbPos] = refFrame->Texture.GetReference();
 				pp.RefFrameList[i].Index7Bits = (UCHAR)dpbPos;
 				pp.RefFrameList[i].AssociatedFlag = !RefFrames[i].bIsLongTerm ? 0 : 1;
@@ -485,6 +506,10 @@ IElectraDecoder::EDecoderError FD3D12VideoDecoder_H264::DecodeSlicesH264(const F
 				pp.FieldOrderCntList[i][0] = RefFrames[i].TopPOC;
 				pp.FieldOrderCntList[i][1] = RefFrames[i].BottomPOC;
 				pp.FrameNumList[i] = !RefFrames[i].bIsLongTerm ? RefFrames[i].FrameNum : RefFrames[i].LongTermFrameIndex;
+				if (refFrame == MissingReferenceFrame)
+				{
+					pp.NonExistingFrameFlags |= (1 << i);
+				}
 			}
 		}
 	}
@@ -591,8 +616,7 @@ IElectraDecoder::EDecoderError FD3D12VideoDecoder_H264::DecodeSlicesH264(const F
 	FrameInfo.IndexInBuffer = TargetFrame->IndexInPictureBuffer;
 	FrameInfo.PTS = InInputAccessUnit.PTS;
 	FrameInfo.UserValue0 = AssociatedUserValue;
-	TArray<ElectraDecodersUtil::MPEG::H264::FOutputFrameInfo> OutputFrameInfos, UnrefFrameInfos;
-	BitstreamParamsH264.DPBPOC.EndFrame(OutputFrameInfos, UnrefFrameInfos, FrameInfo, InSliceInfos[0].NalUnitType, InSliceInfos[0].NalRefIdc, InSliceInfos[0].Header);
+	BitstreamParamsH264.DPBPOC.EndFrame(OutputFrameInfos, UnrefFrameInfos, FrameInfo, InSliceInfos[0].NalUnitType, InSliceInfos[0].NalRefIdc, InSliceInfos[0].Header, false);
 	return HandleOutputListH264(OutputFrameInfos);
 }
 
@@ -601,6 +625,12 @@ IElectraDecoder::EDecoderError FD3D12VideoDecoder_H264::HandleOutputListH264(con
 	TSharedPtr<FVideoDecoderOutputD3D12Electra, ESPMode::ThreadSafe> InDec;
 	for(int32 i=0; i<InOutputFrameInfos.Num(); ++i)
 	{
+		// In case the frame is a missing frame we ignore it.
+		if (InOutputFrameInfos[i].IndexInBuffer < 0)
+		{
+			continue;
+		}
+
 //		UE_LOG(LogD3D12VideoDecodersElectra, Log, TEXT("Output frame %d, %lld"), InOutputFrameInfos[i].IndexInBuffer, (long long int)InOutputFrameInfos[i].PTS.GetTicks());
 		TSharedPtr<FDecodedFrame, ESPMode::ThreadSafe> Frame = DPB->GetFrameAtIndex(InOutputFrameInfos[i].IndexInBuffer);
 		check(Frame.IsValid());
diff --git a/Engine/Plugins/Media/ElectraCodecs/Source/ElectraDecoders/Private/Utils/MPEG/ElectraUtilsMPEGVideo_H264.cpp b/Engine/Plugins/Media/ElectraCodecs/Source/ElectraDecoders/Private/Utils/MPEG/ElectraUtilsMPEGVideo_H264.cpp
--- a/Engine/Plugins/Media/ElectraCodecs/Source/ElectraDecoders/Private/Utils/MPEG/ElectraUtilsMPEGVideo_H264.cpp
+++ b/Engine/Plugins/Media/ElectraCodecs/Source/ElectraDecoders/Private/Utils/MPEG/ElectraUtilsMPEGVideo_H264.cpp
@@ -1181,6 +1181,59 @@ namespace ElectraDecodersUtil
 				CurrentPOC = PreviousPOC;
 			}
 
+			bool FSlicePOCVars::HandleMissingFrames(TArray<FOutputFrameInfo>& OutOutputFrameInfos, TArray<FOutputFrameInfo>& OutUnrefFrameInfos, uint8 InNalUnitType, uint8 InNalRefIdc, const FSliceHeader& InSliceHeader, const FSequenceParameterSet& InSequenceParameterSet)
+			{
+				bool bIsIDR = InNalUnitType == 5;
+				// Do we need to check?
+				if (bLastHadMMCO5 || bIsIDR)
+				{
+					// No need. Just continue.
+					return true;
+				}
+				// Is there a frame missing?
+				if (!(InSliceHeader.frame_num != CurrentPOC.prev_frame_num && InSliceHeader.frame_num != (CurrentPOC.prev_frame_num + 1) % max_frame_num))
+				{
+					// No. Continue.
+					return true;
+				}
+				/*
+					Decoding process for gaps in frame_num, Section 8.2.5.2
+						the current picture is considered to be a picture considered having frame_num inferred to be equal to UnusedShortTermFrameNum,
+						nal_ref_idc inferred to be not equal to 0,
+						nal_unit_type inferred to be not equal to 5,
+						IdrPicFlag inferred to be equal to 0,
+						field_pic_flag inferred to be equal to 0,
+						adaptive_ref_pic_marking_mode_flag inferred to be equal to 0,
+						delta_pic_order_cnt[ 0 ] (if needed) inferred to be equal to 0, and delta_pic_order_cnt[ 1 ] (if needed) inferred to be equal to 0.
+				*/
+				InNalUnitType = 1;
+				InNalRefIdc = 1;
+				FSliceHeader TempSliceHeader(InSliceHeader);
+				TempSliceHeader.delta_pic_order_cnt[0] = TempSliceHeader.delta_pic_order_cnt[1] = 0;
+				TempSliceHeader.adaptive_ref_pic_marking_mode_flag = 0;
+				TempSliceHeader.MemoryManagementControl.Empty();
+				TempSliceHeader.field_pic_flag = 0;
+				TempSliceHeader.long_term_reference_flag = 0;
+				uint32 UnusedShortTermFrameNum = (CurrentPOC.prev_frame_num + 1) % max_frame_num;
+				FOutputFrameInfo NoInfo;
+				NoInfo.bDoNotOutput = true;
+				while(InSliceHeader.frame_num != UnusedShortTermFrameNum)
+				{
+					TempSliceHeader.frame_num = UnusedShortTermFrameNum;
+					if (InSequenceParameterSet.pic_order_cnt_type != 0)
+					{
+						UpdatePOCInternal(InNalUnitType, InNalRefIdc, TempSliceHeader, InSequenceParameterSet);
+					}
+					CurrentPOC.prev_frame_num = UnusedShortTermFrameNum;
+					UnusedShortTermFrameNum = (UnusedShortTermFrameNum + 1) % max_frame_num;
+					if (!EndFrame(OutOutputFrameInfos, OutUnrefFrameInfos, NoInfo, InNalUnitType, InNalRefIdc, TempSliceHeader, true))
+					{
+						return false;
+					}
+				}
+				return true;
+			}
+
 			bool FSlicePOCVars::UpdatePOC(uint8 InNalUnitType, uint8 InNalRefIdc, const FSliceHeader& InSliceHeader, const FSequenceParameterSet& InSequenceParameterSet)
 			{
 				PreviousPOC = CurrentPOC;
@@ -1196,11 +1249,18 @@ namespace ElectraDecodersUtil
 				{
 					CurrentPOC.prev_frame_num = InSliceHeader.frame_num;
 				}
+				// Missing frames should have been covered in HandleMissingFrames() already!
 				if (InSliceHeader.frame_num != CurrentPOC.prev_frame_num && InSliceHeader.frame_num != (CurrentPOC.prev_frame_num + 1) % max_frame_num)
 				{
 					return SetLastError(FString::Printf(TEXT("Gap in frame_num detected. Cannot conceal error.")));
 				}
+				UpdatePOCInternal(InNalUnitType, InNalRefIdc, InSliceHeader, InSequenceParameterSet);
+				return true;
+			}
 
+			void FSlicePOCVars::UpdatePOCInternal(uint8 InNalUnitType, uint8 InNalRefIdc, const FSliceHeader& InSliceHeader, const FSequenceParameterSet& InSequenceParameterSet)
+			{
+				bool bIsIDR = InNalUnitType == 5;
 				// Decoding process for picture order count, Section 8.2.1
 				if (InSequenceParameterSet.pic_order_cnt_type == 0)
 				{
@@ -1361,7 +1421,6 @@ namespace ElectraDecodersUtil
 						it->LongTermPicNum = it->LongTermFrameIndex;
 					}
 				}
-				return true;
 			}
 
 			void FSlicePOCVars::GetCurrentReferenceFrames(TArray<FReferenceFrameListEntry>& OutCurrentReferenceFrames)
@@ -1683,7 +1742,7 @@ namespace ElectraDecodersUtil
 			}
 
 
-			bool FSlicePOCVars::EndFrame(TArray<FOutputFrameInfo>& OutOutputFrameInfos, TArray<FOutputFrameInfo>& OutUnrefFrameInfos, const FOutputFrameInfo& InOutputFrameInfo, uint8 InNalUnitType, uint8 InNalRefIdc, const FSliceHeader& InSliceHeader)
+			bool FSlicePOCVars::EndFrame(TArray<FOutputFrameInfo>& OutOutputFrameInfos, TArray<FOutputFrameInfo>& OutUnrefFrameInfos, const FOutputFrameInfo& InOutputFrameInfo, uint8 InNalUnitType, uint8 InNalRefIdc, const FSliceHeader& InSliceHeader, bool bIsNonExisting)
 			{
 				bool bIsIDR = InNalUnitType == 5;
 				bLastHadMMCO5 = false;
@@ -1696,6 +1755,11 @@ namespace ElectraDecodersUtil
 				dinf->FramePOC = CurrentPOC.FramePOC;
 				dinf->TopPOC = CurrentPOC.TopPOC;
 				dinf->BottomPOC = CurrentPOC.BottomPOC;
+				if ((dinf->bIsNonExisting = bIsNonExisting) == true)
+				{
+					dinf->bHasBeenOutput = true;
+				}
+
 
 				// Decoded reference picture marking process
 				// Section 8.2.5
@@ -1942,7 +2006,7 @@ namespace ElectraDecodersUtil
 				UE_LOG(LogElectraDecoders, Log, TEXT("Current DPB POC"));
 				for(auto& it : FrameDPBInfos)
 				{
-					UE_LOG(LogElectraDecoders, Log, TEXT("fn=%u  poc=%d  ref=%d  ltref=%d  out=%d"), it->FrameNum, it->POC, it->bIsUsedForReference, it->bIsLongTermReference, it->bHasBeenOutput);
+					UE_LOG(LogElectraDecoders, Log, TEXT("fn=%u  poc=%d  ref=%d  ltref=%d  out=%d"), it->FrameNum, it->FramePOC, it->bIsUsedForReference, it->bIsLongTermReference, it->bHasBeenOutput);
 				}
 #endif
 				return true;
diff --git a/Engine/Plugins/Media/ElectraCodecs/Source/ElectraDecoders/Public/Utils/MPEG/ElectraUtilsMPEGVideo_H264.h b/Engine/Plugins/Media/ElectraCodecs/Source/ElectraDecoders/Public/Utils/MPEG/ElectraUtilsMPEGVideo_H264.h
--- a/Engine/Plugins/Media/ElectraCodecs/Source/ElectraDecoders/Public/Utils/MPEG/ElectraUtilsMPEGVideo_H264.h
+++ b/Engine/Plugins/Media/ElectraCodecs/Source/ElectraDecoders/Public/Utils/MPEG/ElectraUtilsMPEGVideo_H264.h
@@ -72,7 +72,7 @@ namespace ElectraDecodersUtil
 				{
 					return static_cast<const uint8*>(GetRemainingData());
 				}
-				
+
 				bool more_rbsp_data()
 				{
 					// There is definitely more data available unless we are in the last byte.
@@ -174,7 +174,7 @@ namespace ElectraDecodersUtil
 					max_dec_frame_buffering = 0;
 					nal_hrd_parameters.Reset();
 					vcl_hrd_parameters.Reset();
-					
+
 					ExpectedDeltaPerPicOrderCntCycle = 0;
 				}
 				int32 GetMaxDPBSize() const;
@@ -255,7 +255,7 @@ namespace ElectraDecodersUtil
 				uint32 max_bits_per_mb_denom;								// ue(v), 0-16, default 1
 				uint32 log2_max_mv_length_horizontal;						// ue(v), 0-15, default 15
 				uint32 log2_max_mv_length_vertical;							// ue(v), 0-15, default 15
-				uint32 max_num_reorder_frames;								// ue(v), 0-max_dec_frame_buffering, 
+				uint32 max_num_reorder_frames;								// ue(v), 0-max_dec_frame_buffering,
 																			//	defaults to: if profile_idc in [44,86,100,110,122,144] and constraint_set3_flag==1 then 0
 																			//               otherwise MaxDpbFrames.
 				uint32 max_dec_frame_buffering;								// ue(v), >= max_num_ref_frames
@@ -468,14 +468,17 @@ namespace ElectraDecodersUtil
 
 				// Call this to get the current short- and long-term reference frames.
 				void GetCurrentReferenceFrames(TArray<FReferenceFrameListEntry>& OutCurrentReferenceFrames);
-				
+
 				// Call this per slice to get the list of references needed for decoding the slice.
 				bool GetReferenceFrameLists(TArray<FReferenceFrameListEntry>& OutReferenceFrameList0, TArray<FReferenceFrameListEntry>& OutReferenceFrameList1, const FSliceHeader& InSliceHeader);
 
 				// Call this on the first slice of an image to be decoded.
 				bool BeginFrame(uint8 InNalUnitType, uint8 InNalRefIdc, const FSliceHeader& InSliceHeader, const FSequenceParameterSet& InSequenceParameterSet, const FPictureParameterSet& InPictureParameterSet);
-
+				// Call this next to see if there are any missing frames between this and the last.
+				bool HandleMissingFrames(TArray<FOutputFrameInfo>& OutOutputFrameInfos, TArray<FOutputFrameInfo>& OutUnrefFrameInfos, uint8 InNalUnitType, uint8 InNalRefIdc, const FSliceHeader& InSliceHeader, const FSequenceParameterSet& InSequenceParameterSet);
+				// Update the POC values for the current frame.
 				bool UpdatePOC(uint8 InNalUnitType, uint8 InNalRefIdc, const FSliceHeader& InSliceHeader, const FSequenceParameterSet& InSequenceParameterSet);
+				// Undo the POC value updates in case the decoding process could not run.
 				void UndoPOCUpdate();
 
 				int32 GetFramePOC() const
@@ -486,11 +489,11 @@ namespace ElectraDecodersUtil
 				{ return CurrentPOC.BottomPOC; }
 
 				// Call this when decoding of the entire frame (all slices) is done.
-				bool EndFrame(TArray<FOutputFrameInfo>& OutOutputFrameInfos, TArray<FOutputFrameInfo>& OutUnrefFrameInfos, const FOutputFrameInfo& InOutputFrameInfo, uint8 InNalUnitType, uint8 InNalRefIdc, const FSliceHeader& InSliceHeader);
+				bool EndFrame(TArray<FOutputFrameInfo>& OutOutputFrameInfos, TArray<FOutputFrameInfo>& OutUnrefFrameInfos, const FOutputFrameInfo& InOutputFrameInfo, uint8 InNalUnitType, uint8 InNalRefIdc, const FSliceHeader& InSliceHeader, bool bIsNonExisting);
 
 				// Call this to get all pending frame infos and clear the structure.
 				void Flush(TArray<FOutputFrameInfo>& OutRemainingFrameInfos, TArray<FOutputFrameInfo>& OutUnrefFrameInfos);
-				
+
 				// Call this to reset this structure.
 				void Reset();
 
@@ -513,6 +516,7 @@ namespace ElectraDecodersUtil
 					bool bHasBeenOutput = false;
 				};
 
+				void UpdatePOCInternal(uint8 InNalUnitType, uint8 InNalRefIdc, const FSliceHeader& InSliceHeader, const FSequenceParameterSet& InSequenceParameterSet);
 				void UpdateRefLists();
 				FSmallestPOC GetSmallestPOC();
 				bool ReorderRefPicList(TArray<FReferenceFrameListEntry>& InOutReferenceFrameList, const FSliceHeader& InSliceHeader, int32 InListNum);
@@ -566,7 +570,7 @@ namespace ElectraDecodersUtil
 				TArray<TSharedPtr<FFrameInDPBInfo>> FrameDPBInfos;
 				TArray<TSharedPtr<FFrameInDPBInfo>> ShortTermRefs;
 				TArray<TSharedPtr<FFrameInDPBInfo>> LongTermRefs;
-				
+
 				FString LastErrorMsg;
 			};
 
