Electra: Preparations to use D3D12 video decoding

#rnx

diff --git a/Engine/Plugins/Media/D3D12VideoDecodersElectra/D3D12VideoDecodersElectra.uplugin b/Engine/Plugins/Media/D3D12VideoDecodersElectra/D3D12VideoDecodersElectra.uplugin
--- a/Engine/Plugins/Media/D3D12VideoDecodersElectra/D3D12VideoDecodersElectra.uplugin
+++ b/Engine/Plugins/Media/D3D12VideoDecodersElectra/D3D12VideoDecodersElectra.uplugin
@@ -9,7 +9,7 @@
   "Category": "Media Players",
   "EnabledByDefault": false,
   "CanContainContent": false,
-  "IsBetaVersion": true,
+  "IsBetaVersion": false,
   "IsExperimentalVersion": false,
   "Installed": false,
   "Modules": [
diff --git a/Engine/Plugins/Media/D3D12VideoDecodersElectra/README.txt b/Engine/Plugins/Media/D3D12VideoDecodersElectra/README.txt
--- a/Engine/Plugins/Media/D3D12VideoDecodersElectra/README.txt
+++ b/Engine/Plugins/Media/D3D12VideoDecodersElectra/README.txt
@@ -12,3 +12,40 @@ H.264 / AVC
 
 H.265 / HEVC
 - Only Main and Main10 profiles up to level 6.3 are supported
+
+
+
+
+NOTE on usage:
+==============
+
+By default, even with this plugin enabled, its use is disabled on Windows due to the
+vast combinations of GPUs and driver implementations.
+It can be enabled via CVar / .ini file
+To enable via ini file add
+	[ConsoleVariables]
+	ElectraDecoders.bDoNotUseD3D12Video=false
+to the respective platform's ini file.
+
+To enable by default at compile time, override the method
+		protected virtual bool bIsDefaultIgnoredOnPlatform
+		{
+			get
+			{
+				return false;
+			}
+		}
+from the base D3D12VideoDecodersElectra.Build.cs build script in the platform build
+scripts where you want it enabled.
+
+For testing you should confirm operation by setting the console variable
+	`ElectraDecoders.bDoNotUseD3D12Video`
+to `false`.
+
+
+
+To disable the plugin right at Startup() you can use the CVar / ini setting:
+	[ConsoleVariables]
+	ElectraDecoders.bDisableD3D12Video=true
+
+This acts as a master switch. When disabled `bDoNotUseD3D12Video` has no effect.
diff --git a/Engine/Plugins/Media/D3D12VideoDecodersElectra/Source/D3D12VideoDecodersElectra.Build.cs b/Engine/Plugins/Media/D3D12VideoDecodersElectra/Source/D3D12VideoDecodersElectra.Build.cs
--- a/Engine/Plugins/Media/D3D12VideoDecodersElectra/Source/D3D12VideoDecodersElectra.Build.cs
+++ b/Engine/Plugins/Media/D3D12VideoDecodersElectra/Source/D3D12VideoDecodersElectra.Build.cs
@@ -4,8 +4,25 @@ namespace UnrealBuildTool.Rules
 {
 	public class D3D12VideoDecodersElectra : ModuleRules
 	{
+		protected virtual bool bIsDefaultDisabledOnPlatform
+		{
+			get
+			{
+				return false;
+			}
+		}
+		protected virtual bool bIsDefaultIgnoredOnPlatform
+		{
+			get
+			{
+				return true;
+			}
+		}
+
 		public D3D12VideoDecodersElectra(ReadOnlyTargetRules Target) : base(Target)
 		{
+			PublicDefinitions.Add("ELECTRA_DECODERS_D3D12VIDEO_DISABLED_ON_PLATFORM=" + (bIsDefaultDisabledOnPlatform ? "1" : "0"));
+			PublicDefinitions.Add("ELECTRA_DECODERS_D3D12VIDEO_IGNORED_ON_PLATFORM=" + (bIsDefaultIgnoredOnPlatform ? "1" : "0"));
 			PrivateDependencyModuleNames.AddRange(
 				new string[] {
 					"Core",
diff --git a/Engine/Plugins/Media/D3D12VideoDecodersElectra/Source/Private/VideoDecoder_D3D12.cpp b/Engine/Plugins/Media/D3D12VideoDecodersElectra/Source/Private/VideoDecoder_D3D12.cpp
--- a/Engine/Plugins/Media/D3D12VideoDecodersElectra/Source/Private/VideoDecoder_D3D12.cpp
+++ b/Engine/Plugins/Media/D3D12VideoDecodersElectra/Source/Private/VideoDecoder_D3D12.cpp
@@ -7,6 +7,8 @@
 #include "Features/IModularFeature.h"
 #include "RHI.h"
 
+#include "HAL/IConsoleManager.h"
+
 #include "IElectraCodecFactory.h"
 #include "IElectraCodecFactoryModule.h"
 #include "IElectraCodecRegistry.h"
@@ -20,6 +22,36 @@
 
 namespace ElectraVideoDecodersD3D12Video
 {
+#ifndef ELECTRA_DECODERS_D3D12VIDEO_DISABLED_ON_PLATFORM
+#error "C"
+#define ELECTRA_DECODERS_D3D12VIDEO_DISABLED_ON_PLATFORM 1
+#endif
+#ifndef ELECTRA_DECODERS_D3D12VIDEO_IGNORED_ON_PLATFORM
+#error "B"
+#define ELECTRA_DECODERS_D3D12VIDEO_IGNORED_ON_PLATFORM 1
+#endif
+
+#if ELECTRA_DECODERS_D3D12VIDEO_DISABLED_ON_PLATFORM
+static bool bDisableThisDecoder = true;
+#else
+static bool bDisableThisDecoder = false;
+#endif
+
+#if ELECTRA_DECODERS_D3D12VIDEO_IGNORED_ON_PLATFORM
+static bool bDoNotUseThisDecoder = true;
+#else
+static bool bDoNotUseThisDecoder = false;
+#endif
+
+FAutoConsoleVariableRef CVarElectraDecoderD3D12VideoDisable(
+	TEXT("ElectraDecoders.bDisableD3D12Video"),
+	bDisableThisDecoder,
+	TEXT("Globally disable the use of the D3D12 native video decoder"));
+
+FAutoConsoleVariableRef CVarElectraDecoderD3D12DoNotUse(
+	TEXT("ElectraDecoders.bDoNotUseD3D12Video"),
+	bDoNotUseThisDecoder,
+	TEXT("Do not use the D3D12 native video decoder on this platform"));
 
 int32 FCodecFormatHelper::FindSupportedFormats(ID3D12Device* InD3D12Device)
 {
@@ -187,7 +219,7 @@ const FCodecFormatHelper::FCodecInfo* FCodecFormatHelper::HaveFormat(ECodecType
 	for(int32 i=0, iMax=CodecInfos.Num(); i<iMax; ++i)
 	{
 		if (CodecInfos[i].CodecType == InType &&
-		   ((InNumBits == 8 && CodecInfos[i].b10Bit == false) || (InNumBits == 10 && CodecInfos[i].b10Bit == true)))
+			((InNumBits == 8 && CodecInfos[i].b10Bit == false) || (InNumBits == 10 && CodecInfos[i].b10Bit == true)))
 		{
 			return &CodecInfos[i];
 		}
@@ -385,6 +417,13 @@ public:
 					UE_LOG(LogD3D12VideoDecodersElectra, Log, TEXT("Decode tier 2 is needed, but tier %d was returned."), (int)OutSupport.DecodeTier);
 					Codec = nullptr;
 				}
+
+				// Do a custom platform capability check.
+				if (!FD3D12VideoDecoder::CheckPlatformDecodeCapabilities(OutSupport, ci, InOptions))
+				{
+					UE_LOG(LogD3D12VideoDecodersElectra, Log, TEXT("Platform rejected decoding of %d*%d @ %d/%d fps"), (int)Width, (int)Height, (int)fps_n, (int)fps_d);
+					Codec = nullptr;
+				}
 			}
 			else
 			{
@@ -397,6 +436,11 @@ public:
 
 	int32 SupportsFormat(const FString& InCodecFormat, bool bInEncoder, const TMap<FString, FVariant>& InOptions) const override
 	{
+		if (bDoNotUseThisDecoder)
+		{
+			return 0;
+		}
+
 		// Encoder? Not supported here!
 		if (bInEncoder)
 		{
@@ -429,6 +473,11 @@ public:
 
 	TSharedPtr<IElectraDecoder, ESPMode::ThreadSafe> CreateDecoderForFormat(const FString& InCodecFormat, const TMap<FString, FVariant>& InOptions, TSharedPtr<IElectraDecoderResourceDelegate, ESPMode::ThreadSafe> InResourceDelegate) override
 	{
+		if (bDoNotUseThisDecoder)
+		{
+			return nullptr;
+		}
+
 		HRESULT Result;
 
 		// Do this under lock as it may be possible that the D3D device changed and we have to rebuild the codec list.
@@ -507,6 +556,10 @@ public:
 				break;
 			}
 		}
+		if (New.IsValid())
+		{
+			UE_LOG(LogD3D12VideoDecodersElectra, Verbose, TEXT("Created a D3D12 video decoder."));
+		}
 		return New;
 	}
 
@@ -587,7 +640,7 @@ bool FD3D12VideoDecoder::ResetToCleanStart()
 	VideoDecoderCommandQueue.SafeRelease();
 	VideoDecoderSync.Reset();
 
-	CurrentConfig.VideoDecoderHeap.SafeRelease();
+	CurrentConfig.Reset();
 	StatusReportFeedbackNumber = 0;
 	return true;
 }
@@ -839,25 +892,28 @@ bool FD3D12VideoDecoder::InternalDecoderCreate()
 	return true;
 }
 
-bool FD3D12VideoDecoder::CreateDecoderHeapAndDPB(int32 InDPBSize, int32 InNumFrames, int32 InImageSizeAlignment)
+bool FD3D12VideoDecoder::CreateDecoderHeap(int32 InDPBSize, int32 InMaxWidth, int32 InMaxHeight, int32 InImageSizeAlignment)
 {
-	if (InDPBSize <= 0 || InNumFrames <= 0 || InNumFrames < InDPBSize)
+	if (InDPBSize <= 0)
 	{
 		return PostError(0, TEXT("DPB size is invalid"), ERRCODE_INTERNAL_FAILED_TO_CREATE_DECODER);
 	}
-	if (InNumFrames > FFrameDecodeResource::kMaxRefFrames)
-	{
-		return PostError(0, TEXT("Too many frames requested than fit into managing structure."), ERRCODE_INTERNAL_FAILED_TO_CREATE_DECODER);
-	}
 
 	HRESULT Result;
 	TRefCountPtr<ID3D12VideoDecoderHeap> NewHeap;
 	D3D12_VIDEO_DECODER_HEAP_DESC HeapDesc {};
 	HeapDesc.NodeMask = GetNodeMask();
 	HeapDesc.Configuration = DecodeSupport.Configuration;
-	const uint64 Alignment = (uint64) InImageSizeAlignment;
-	HeapDesc.DecodeWidth = Align(DecodeSupport.Width, Alignment);
-	HeapDesc.DecodeHeight = Align(DecodeSupport.Height, Alignment);
+	const uint32 Alignment = (uint32) InImageSizeAlignment;
+	const uint32 AlignedWidth = Align(InMaxWidth, Alignment);
+	const uint32 AlignedHeight = Align(InMaxHeight, Alignment);
+#if 1
+	HeapDesc.DecodeWidth = AlignedWidth;
+	HeapDesc.DecodeHeight = AlignedHeight;
+#else
+	HeapDesc.DecodeWidth = DecodeSupport.Width;
+	HeapDesc.DecodeHeight = DecodeSupport.Height;
+#endif
 	check(!CodecInfo.PixelFormats.IsEmpty());
 	HeapDesc.Format = CodecInfo.PixelFormats[0];
 	// best not to set those
@@ -868,22 +924,25 @@ bool FD3D12VideoDecoder::CreateDecoderHeapAndDPB(int32 InDPBSize, int32 InNumFra
 	{
 		return PostError(Result, TEXT("CreateVideoDecoderHeap() failed"), ERRCODE_INTERNAL_FAILED_TO_CREATE_BUFFER);
 	}
-
-	if (!CreateDPB(DPB, HeapDesc.DecodeWidth, HeapDesc.DecodeHeight, InNumFrames, HeapDesc.Format))
-	{
-		return false;
-	}
-
 	CurrentConfig.VideoDecoderHeap = NewHeap;
+	CurrentConfig.VideoDecoderDPBWidth = InMaxWidth;
+	CurrentConfig.VideoDecoderDPBHeight = InMaxHeight;
 	CurrentConfig.MaxNumInDPB = InDPBSize;
-	CurrentConfig.MaxDecodedWidth = HeapDesc.DecodeWidth;
-	CurrentConfig.MaxDecodedHeight = HeapDesc.DecodeHeight;
 	return true;
 }
 
 
-bool FD3D12VideoDecoder::CreateDPB(TSharedPtr<FDecodedPictureBuffer, ESPMode::ThreadSafe>& OutDPB, int32 InMaxWidth, int32 InMaxHeight, int32 InNumFrames, DXGI_FORMAT InFormat)
+bool FD3D12VideoDecoder::CreateDPB(TSharedPtr<FDecodedPictureBuffer, ESPMode::ThreadSafe>& OutDPB, int32 InMaxWidth, int32 InMaxHeight, int32 InImageSizeAlignment, int32 InNumFrames)
 {
+	if (InNumFrames <= 0)
+	{
+		return PostError(0, TEXT("Bad number of frames"), ERRCODE_INTERNAL_FAILED_TO_CREATE_DECODER);
+	}
+	if (InNumFrames > FFrameDecodeResource::kMaxRefFrames)
+	{
+		return PostError(0, TEXT("Too many frames requested than fit into managing structure."), ERRCODE_INTERNAL_FAILED_TO_CREATE_DECODER);
+	}
+
 	TSharedPtr<FDecodedPictureBuffer, ESPMode::ThreadSafe> newdpb = MakeShared<FDecodedPictureBuffer, ESPMode::ThreadSafe>();
 	newdpb->Frames.SetNum(InNumFrames);
 	D3D12_HEAP_PROPERTIES heapProps {};
@@ -893,14 +952,19 @@ bool FD3D12VideoDecoder::CreateDPB(TSharedPtr<FDecodedPictureBuffer, ESPMode::Th
 	heapProps.CreationNodeMask =
 	heapProps.VisibleNodeMask = GetNodeMask();
 
+	const uint32 Alignment = (uint32) InImageSizeAlignment;
+	const uint32 AlignedWidth = Align(InMaxWidth, Alignment);
+	const uint32 AlignedHeight = Align(InMaxHeight, Alignment);
+
+	check(!CodecInfo.PixelFormats.IsEmpty());
 	D3D12_RESOURCE_DESC desc {};
 	desc.Dimension = D3D12_RESOURCE_DIMENSION_TEXTURE2D;
 	desc.Alignment = 0;
-	desc.Width = InMaxWidth;
-	desc.Height = InMaxHeight;
+	desc.Width = AlignedWidth;
+	desc.Height = AlignedHeight;
 	desc.DepthOrArraySize = 1;
 	desc.MipLevels = 1;
-	desc.Format = InFormat;
+	desc.Format = CodecInfo.PixelFormats[0];
 	desc.SampleDesc.Count = 1;
 	desc.SampleDesc.Quality = 0;
 	desc.Layout = D3D12_TEXTURE_LAYOUT_UNKNOWN;
@@ -926,6 +990,8 @@ bool FD3D12VideoDecoder::CreateDPB(TSharedPtr<FDecodedPictureBuffer, ESPMode::Th
 		newdpb->AvailableQueue.Push(it);
 	}
 	OutDPB = MoveTemp(newdpb);
+	CurrentConfig.MaxDecodedWidth = InMaxWidth;
+	CurrentConfig.MaxDecodedHeight = InMaxHeight;
 	return true;
 }
 
@@ -983,17 +1049,28 @@ static TSharedPtr<FD3D12VideoDecoderFactory, ESPMode::ThreadSafe> Self;
 /*********************************************************************************************************************/
 void FD3D12VideoDecoder::Startup()
 {
-	// Make sure the codec factory module has been loaded.
-	FModuleManager::Get().LoadModule(TEXT("ElectraCodecFactory"));
+	if (!ElectraVideoDecodersD3D12Video::bDisableThisDecoder)
+	{
+		// Make sure the codec factory module has been loaded.
+		FModuleManager::Get().LoadModule(TEXT("ElectraCodecFactory"));
 
-	TUniquePtr<ElectraVideoDecodersD3D12Video::FCodecFormatHelper> FormatHelper(new ElectraVideoDecodersD3D12Video::FCodecFormatHelper);
-	// Not a single supported format?
-	if (FormatHelper->FindSupportedFormats(nullptr) != 0)
+		TUniquePtr<ElectraVideoDecodersD3D12Video::FCodecFormatHelper> FormatHelper(new ElectraVideoDecodersD3D12Video::FCodecFormatHelper);
+		// Not a single supported format?
+		if (FormatHelper->FindSupportedFormats(nullptr) != 0)
+		{
+			// Create a factory with the current formats.
+			ElectraVideoDecodersD3D12Video::Self = MakeShared<ElectraVideoDecodersD3D12Video::FD3D12VideoDecoderFactory, ESPMode::ThreadSafe>(MoveTemp(FormatHelper));
+			// Register as modular feature.
+			IModularFeatures::Get().RegisterModularFeature(IElectraCodecFactoryModule::GetModularFeatureName(), ElectraVideoDecodersD3D12Video::Self.Get());
+		}
+		else
+		{
+			UE_LOG(LogD3D12VideoDecodersElectra, Log, TEXT("D3D12 video decoding will not be used since no supported format was found."));
+		}
+	}
+	else
 	{
-		// Create a factory with the current formats.
-		ElectraVideoDecodersD3D12Video::Self = MakeShared<ElectraVideoDecodersD3D12Video::FD3D12VideoDecoderFactory, ESPMode::ThreadSafe>(MoveTemp(FormatHelper));
-		// Register as modular feature.
-		IModularFeatures::Get().RegisterModularFeature(IElectraCodecFactoryModule::GetModularFeatureName(), ElectraVideoDecodersD3D12Video::Self.Get());
+		UE_LOG(LogD3D12VideoDecodersElectra, Log, TEXT("D3D12 video decoding will not be used since it is disabled."));
 	}
 }
 
diff --git a/Engine/Plugins/Media/D3D12VideoDecodersElectra/Source/Private/VideoDecoder_D3D12_Common.h b/Engine/Plugins/Media/D3D12VideoDecodersElectra/Source/Private/VideoDecoder_D3D12_Common.h
--- a/Engine/Plugins/Media/D3D12VideoDecodersElectra/Source/Private/VideoDecoder_D3D12_Common.h
+++ b/Engine/Plugins/Media/D3D12VideoDecodersElectra/Source/Private/VideoDecoder_D3D12_Common.h
@@ -340,6 +340,7 @@ public:
 class FD3D12VideoDecoder : public IElectraDecoder
 {
 public:
+	static bool D3D12VIDEODECODERSELECTRA_API CheckPlatformDecodeCapabilities(D3D12_FEATURE_DATA_VIDEO_DECODE_SUPPORT& InOutDecodeSupport, const ElectraDecodersUtil::FMimeTypeVideoCodecInfo& InCodecInfo, const TMap<FString, FVariant>& InOptions);
 	FD3D12VideoDecoder(const FCodecFormatHelper::FCodecInfo& InCodecInfo, const D3D12_FEATURE_DATA_VIDEO_DECODE_SUPPORT& InDecodeSupport, const TMap<FString, FVariant>& InOptions, TSharedPtr<IElectraDecoderResourceDelegate, ESPMode::ThreadSafe> InResourceDelegate, const TRefCountPtr<ID3D12Device>& InD3D12Device, const TRefCountPtr<ID3D12VideoDevice>& InVideoDevice, uint32 InVideoDeviceNodeIndex);
 	virtual ~FD3D12VideoDecoder();
 protected:
@@ -434,17 +435,28 @@ protected:
 
 	struct FDecoderConfiguration
 	{
+		void Reset()
+		{
+			MaxDecodedWidth = 0;
+			MaxDecodedHeight = 0;
+			MaxNumInDPB = 0;
+			VideoDecoderHeap.SafeRelease();
+			VideoDecoderDPBWidth = 0;
+			VideoDecoderDPBHeight = 0;
+		}
 		int32 MaxDecodedWidth = 0;
 		int32 MaxDecodedHeight = 0;
 		int32 MaxNumInDPB = 0;
 		TRefCountPtr<ID3D12VideoDecoderHeap> VideoDecoderHeap;
+		int32 VideoDecoderDPBWidth = 0;
+		int32 VideoDecoderDPBHeight = 0;
 	};
 
 	bool InternalDecoderCreate();
 	void ReturnAllFrames();
 
-	bool CreateDecoderHeapAndDPB(int32 InDPBSize, int32 InNumFrames, int32 InImageSizeAlignment);
-	bool CreateDPB(TSharedPtr<FDecodedPictureBuffer, ESPMode::ThreadSafe>& OutDPB, int32 InMaxWidth, int32 InMaxHeight, int32 InNumFrames, DXGI_FORMAT InFormat);
+	bool CreateDecoderHeap(int32 InDPBSize, int32 InMaxWidth, int32 InMaxHeight, int32 InImageSizeAlignment);
+	bool CreateDPB(TSharedPtr<FDecodedPictureBuffer, ESPMode::ThreadSafe>& OutDPB, int32 InMaxWidth, int32 InMaxHeight, int32 InImageSizeAlignment, int32 InNumFrames);
 
 	constexpr uint32 GetNodeMask() const
 	{ return VideoDeviceNodeIndex; }
diff --git a/Engine/Plugins/Media/D3D12VideoDecodersElectra/Source/Private/VideoDecoder_D3D12_H264.cpp b/Engine/Plugins/Media/D3D12VideoDecodersElectra/Source/Private/VideoDecoder_D3D12_H264.cpp
--- a/Engine/Plugins/Media/D3D12VideoDecodersElectra/Source/Private/VideoDecoder_D3D12_H264.cpp
+++ b/Engine/Plugins/Media/D3D12VideoDecodersElectra/Source/Private/VideoDecoder_D3D12_H264.cpp
@@ -39,9 +39,9 @@ IElectraDecoder::ECSDCompatibility FD3D12VideoDecoder_H264::IsCompatibleWith(con
 	}
 
 	const ElectraDecodersUtil::MPEG::H264::FSequenceParameterSet sps = temp.SPSs.CreateConstIterator().Value();
-	int32 NewDPBSize = sps.GetDPBSize();
-	int32 NewWidth = sps.GetWidth();
-	int32 NewHeight = sps.GetHeight();
+	int32 NewWidth, NewHeight, NewDPBSize;
+	NewDPBSize = sps.GetDPBSize();
+	sps.GetDisplaySize(NewWidth, NewHeight);
 	if (NewDPBSize > CurrentConfig.MaxNumInDPB ||
 		NewWidth > CurrentConfig.MaxDecodedWidth ||
 		NewHeight > CurrentConfig.MaxDecodedHeight)
@@ -217,13 +217,32 @@ IElectraDecoder::EDecoderError FD3D12VideoDecoder_H264::DecodeAccessUnit(const F
 		// decoding parameters changed.
 		if (bIsIDR)
 		{
-			// Note: Should we add a check if resolution, DPB or other relevant values have changed?
-			//       This should not actually happen given that the upper layer decoder handles stream switches.
+			const int32 Alignment = 16;
+			const int32 DPBSize = spsPtr->GetDPBSize();
+			int32 dw, dh;
+			spsPtr->GetDisplaySize(dw, dh);
+
+			// Check if the decoder heap parameters have changed such that we have to create a new one.
+			if (DPBSize != CurrentConfig.MaxNumInDPB || dw != CurrentConfig.VideoDecoderDPBWidth || dh != CurrentConfig.VideoDecoderDPBHeight)
+			{
+				CurrentConfig.VideoDecoderHeap.SafeRelease();
+			}
 			if (!CurrentConfig.VideoDecoderHeap.IsValid())
 			{
-				int32 DPBSize = spsPtr->GetDPBSize();
-				int32 NumFrames = spsPtr->GetDPBSize() + 2;	// 1 extra for the current frame that's not in the DPB yet, and 1 extra that acts as a 'missing' frame.
-				if (!CreateDecoderHeapAndDPB(DPBSize, NumFrames, 16))
+				if (!CreateDecoderHeap(DPBSize, dw, dh, Alignment))
+				{
+					return IElectraDecoder::EDecoderError::Error;
+				}
+			}
+
+			if (!DPB.IsValid())
+			{
+				// As far as the decoded frames go, their size can be the maximum that is required
+				// for this stream (the largest resolution).
+				const int32 Width = (int32) DecodeSupport.Width;
+				const int32 Height = (int32) DecodeSupport.Height;
+				const int32 NumFrames = spsPtr->GetDPBSize() + 2;	// 1 extra for the current frame that's not in the DPB yet, and 1 extra that acts as a 'missing' frame.
+				if (!CreateDPB(DPB, Width, Height, Alignment, NumFrames))
 				{
 					return IElectraDecoder::EDecoderError::Error;
 				}
@@ -233,9 +252,7 @@ IElectraDecoder::EDecoderError FD3D12VideoDecoder_H264::DecodeAccessUnit(const F
 					PostError(0, TEXT("Could not create empty frame used to fill in for missing frames"), ERRCODE_INTERNAL_FAILED_TO_DECODE);
 					return IElectraDecoder::EDecoderError::Error;
 				}
-
 			}
-
 		}
 		IElectraDecoder::EDecoderError Error = DecodeSlicesH264(InInputAccessUnit, SliceInfos, *spsPtr, *ppsPtr);
 		return Error;
@@ -561,7 +578,9 @@ IElectraDecoder::EDecoderError FD3D12VideoDecoder_H264::DecodeSlicesH264(const F
 	isa.ReferenceFrames.NumTexture2Ds = FFrameDecodeResource::kMaxRefFrames;
 	isa.ReferenceFrames.ppTexture2Ds = fdr->ReferenceFrameList;
 	isa.ReferenceFrames.pSubresources = fdr->ReferenceFrameListSubRes;
+#if PLATFORM_WINDOWS
 	isa.ReferenceFrames.ppHeaps = nullptr;
+#endif
 
 	IElectraDecoder::EDecoderError decres = ExecuteCommonDecode(isa, osa);
 	if (decres != IElectraDecoder::EDecoderError::None)
@@ -607,8 +626,8 @@ IElectraDecoder::EDecoderError FD3D12VideoDecoder_H264::DecodeSlicesH264(const F
 	InDec->ExtraValues.Emplace(TEXT("dxversion"), FVariant((int64) 12000));
 	InDec->ExtraValues.Emplace(TEXT("sw"), FVariant(false));
 	InDec->ExtraValues.Emplace(TEXT("codec"), FVariant(TEXT("avc")));
-	//InDec->ExtraValues.Emplace(TEXT("pixfmt"), FVariant((int64)EElectraDecoderPlatformPixelFormat::NV12));
-	//InDec->ExtraValues.Emplace(TEXT("pixenc"), FVariant((int64)EElectraDecoderPlatformPixelEncoding::Native));
+	InDec->ExtraValues.Emplace(TEXT("pixfmt"), FVariant((int64)EElectraDecoderPlatformPixelFormat::NV12));
+	InDec->ExtraValues.Emplace(TEXT("pixenc"), FVariant((int64)EElectraDecoderPlatformPixelEncoding::Native));
 	FramesInDecoder.Emplace(MoveTemp(InDec));
 
 	// Update the simulation DPB with the new decoded frame.
diff --git a/Engine/Plugins/Media/D3D12VideoDecodersElectra/Source/Private/VideoDecoder_D3D12_H265.cpp b/Engine/Plugins/Media/D3D12VideoDecodersElectra/Source/Private/VideoDecoder_D3D12_H265.cpp
--- a/Engine/Plugins/Media/D3D12VideoDecodersElectra/Source/Private/VideoDecoder_D3D12_H265.cpp
+++ b/Engine/Plugins/Media/D3D12VideoDecodersElectra/Source/Private/VideoDecoder_D3D12_H265.cpp
@@ -50,9 +50,9 @@ IElectraDecoder::ECSDCompatibility FD3D12VideoDecoder_H265::IsCompatibleWith(con
 	{
 		return IElectraDecoder::ECSDCompatibility::DrainAndReset;
 	}
-	int32 NewDPBSize = sps.GetDPBSize();
-	int32 NewWidth = sps.GetWidth();
-	int32 NewHeight = sps.GetHeight();
+	int32 NewWidth, NewHeight, NewDPBSize;
+	NewDPBSize = sps.GetDPBSize();
+	sps.GetDisplaySize(NewWidth, NewHeight);
 	if (NewDPBSize > CurrentConfig.MaxNumInDPB ||
 		NewWidth > CurrentConfig.MaxDecodedWidth ||
 		NewHeight > CurrentConfig.MaxDecodedHeight)
@@ -252,14 +252,35 @@ IElectraDecoder::EDecoderError FD3D12VideoDecoder_H265::DecodeAccessUnit(const F
 		// decoding parameters changed.
 		if (SliceInfos[0].Header.bIsIRAP)
 		{
-			// Note: Should we add a check if resolution, DPB or other relevant values have changed?
-			//       This should not actually happen given that the upper layer decoder
-			//       handles stream switches.
+			const int32 DPBSize = spsPtr->GetDPBSize();
+			int32 dw, dh;
+			spsPtr->GetDisplaySize(dw, dh);
+
+			// Check if the decoder heap parameters have changed such that we have to create a new one.
+			if (DPBSize != CurrentConfig.MaxNumInDPB || dw != CurrentConfig.VideoDecoderDPBWidth || dh != CurrentConfig.VideoDecoderDPBHeight)
+			{
+				CurrentConfig.VideoDecoderHeap.SafeRelease();
+			}
 			if (!CurrentConfig.VideoDecoderHeap.IsValid())
 			{
-				int32 DPBSize = spsPtr->GetDPBSize();
-				int32 NumFrames = spsPtr->GetDPBSize() + 2;	// 1 extra for the current frame that's not in the DPB yet, and 1 extra that acts as a 'missing' frame.
-				if (!CreateDecoderHeapAndDPB(DPBSize, NumFrames, 64))
+				const int32 Alignment = spsPtr->GetMinCbSizeY();
+				if (!CreateDecoderHeap(DPBSize, dw, dh, Alignment))
+				{
+					return IElectraDecoder::EDecoderError::Error;
+				}
+			}
+
+			if (!DPB.IsValid())
+			{
+				// Use the maximum MinCbSizeY value for image alignment as stipulated in the DXVA HEVC documentation.
+				const int32 Alignment = 64;
+
+				// As far as the decoded frames go, their size can be the maximum that is required
+				// for this stream (the largest resolution).
+				const int32 Width = (int32) DecodeSupport.Width;
+				const int32 Height = (int32) DecodeSupport.Height;
+				const int32 NumFrames = spsPtr->GetDPBSize() + 2;	// 1 extra for the current frame that's not in the DPB yet, and 1 extra that acts as a 'missing' frame.
+				if (!CreateDPB(DPB, Width, Height, Alignment, NumFrames))
 				{
 					return IElectraDecoder::EDecoderError::Error;
 				}
@@ -270,7 +291,6 @@ IElectraDecoder::EDecoderError FD3D12VideoDecoder_H265::DecodeAccessUnit(const F
 					return IElectraDecoder::EDecoderError::Error;
 				}
 			}
-
 		}
 		IElectraDecoder::EDecoderError Error = DecodeSlicesH265(InInputAccessUnit, SliceInfos, *spsPtr, *ppsPtr);
 		// When the frame was decoded and there is an EOS/EOB NALU present we need to apply it now.
@@ -701,7 +721,9 @@ IElectraDecoder::EDecoderError FD3D12VideoDecoder_H265::DecodeSlicesH265(const F
 	isa.ReferenceFrames.NumTexture2Ds = FFrameDecodeResource::kMaxRefFrames;
 	isa.ReferenceFrames.ppTexture2Ds = fdr->ReferenceFrameList;
 	isa.ReferenceFrames.pSubresources = fdr->ReferenceFrameListSubRes;
+#if PLATFORM_WINDOWS
 	isa.ReferenceFrames.ppHeaps = nullptr;
+#endif
 
 	IElectraDecoder::EDecoderError decres = ExecuteCommonDecode(isa, osa);
 	if (decres != IElectraDecoder::EDecoderError::None)
@@ -748,15 +770,15 @@ IElectraDecoder::EDecoderError FD3D12VideoDecoder_H265::DecodeSlicesH265(const F
 	{
 		InDec->BufferFormat = EElectraDecoderPlatformPixelFormat::NV12;
 		InDec->BufferEncoding = EElectraDecoderPlatformPixelEncoding::Native;
-		//InDec->ExtraValues.Emplace(TEXT("pixfmt"), FVariant((int64)EElectraDecoderPlatformPixelFormat::NV12));
-		//InDec->ExtraValues.Emplace(TEXT("pixenc"), FVariant((int64)EElectraDecoderPlatformPixelEncoding::Native));
+		InDec->ExtraValues.Emplace(TEXT("pixfmt"), FVariant((int64)EElectraDecoderPlatformPixelFormat::NV12));
+		InDec->ExtraValues.Emplace(TEXT("pixenc"), FVariant((int64)EElectraDecoderPlatformPixelEncoding::Native));
 	}
 	else
 	{
 		InDec->BufferFormat = EElectraDecoderPlatformPixelFormat::P010;
 		InDec->BufferEncoding = EElectraDecoderPlatformPixelEncoding::Native;
-		//InDec->ExtraValues.Emplace(TEXT("pixfmt"), FVariant((int64)EElectraDecoderPlatformPixelFormat::P010));
-		//InDec->ExtraValues.Emplace(TEXT("pixenc"), FVariant((int64)EElectraDecoderPlatformPixelEncoding::Native));
+		InDec->ExtraValues.Emplace(TEXT("pixfmt"), FVariant((int64)EElectraDecoderPlatformPixelFormat::P010));
+		InDec->ExtraValues.Emplace(TEXT("pixenc"), FVariant((int64)EElectraDecoderPlatformPixelEncoding::Native));
 	}
 	InSequenceParameterSet.GetAspect(InDec->AspectW, InDec->AspectH);
 	auto FrameRate = InSequenceParameterSet.GetTiming();
diff --git a/Engine/Plugins/Media/D3D12VideoDecodersElectra/Source/Private/Windows/WindowsPlatformD3D12VideoDecodeCapabilities.cpp b/Engine/Plugins/Media/D3D12VideoDecodersElectra/Source/Private/Windows/WindowsPlatformD3D12VideoDecodeCapabilities.cpp
new file mode 100644
--- /dev/null
+++ b/Engine/Plugins/Media/D3D12VideoDecodersElectra/Source/Private/Windows/WindowsPlatformD3D12VideoDecodeCapabilities.cpp
@@ -0,0 +1,13 @@
+// Copyright Epic Games, Inc. All Rights Reserved.
+
+#include "VideoDecoder_D3D12_Common.h"
+
+namespace ElectraVideoDecodersD3D12Video
+{
+
+bool FD3D12VideoDecoder::CheckPlatformDecodeCapabilities(D3D12_FEATURE_DATA_VIDEO_DECODE_SUPPORT& InOutDecodeSupport, const ElectraDecodersUtil::FMimeTypeVideoCodecInfo& InCodecInfo, const TMap<FString, FVariant>& InOptions)
+{
+	return true;
+}
+
+}
diff --git a/Engine/Plugins/Media/ElectraCodecs/Source/ElectraDecoders/Private/Utils/MPEG/ElectraUtilsMPEGVideo_H264.cpp b/Engine/Plugins/Media/ElectraCodecs/Source/ElectraDecoders/Private/Utils/MPEG/ElectraUtilsMPEGVideo_H264.cpp
--- a/Engine/Plugins/Media/ElectraCodecs/Source/ElectraDecoders/Private/Utils/MPEG/ElectraUtilsMPEGVideo_H264.cpp
+++ b/Engine/Plugins/Media/ElectraCodecs/Source/ElectraDecoders/Private/Utils/MPEG/ElectraUtilsMPEGVideo_H264.cpp
@@ -286,11 +286,15 @@ namespace ElectraDecodersUtil
 				}
 				if (vui_parameters_present_flag != 0 && bitstream_restriction_flag != 0)
 				{
+				/*
+					Removed for now as this can happen despite it being invalid.
+
 					if ((int32)max_dec_frame_buffering > MaxSize)
 					{
 						UE_LOG(LogElectraDecoders, Error, TEXT("max_dec_frame_buffering in vui is larger than max dpb size for this level"));
 						return -1;
 					}
+				*/
 					uint32 SizeFromVUI = Max(1U, max_dec_frame_buffering);
 					MaxSize = (int32)SizeFromVUI;
 				}
@@ -304,6 +308,13 @@ namespace ElectraDecodersUtil
 			{
 				return (pic_height_in_map_units_minus1 + 1) * 16;
 			}
+			void FSequenceParameterSet::GetDisplaySize(int32& OutWidth, int32& OutHeight) const
+			{
+				int32 cl,cr,ct,cb;
+				GetCrop(cl,cr,ct,cb);
+				OutWidth = GetWidth() - cl - cr;
+				OutHeight = GetHeight() - ct - cb;
+			}
 			void FSequenceParameterSet::GetCrop(int32& OutLeft, int32& OutRight, int32& OutTop, int32& OutBottom) const
 			{
 				if (frame_cropping_flag)
@@ -490,8 +501,8 @@ namespace ElectraDecodersUtil
 					sps.frame_crop_bottom_offset = br.ue_v();
 				}
 				uint32 MaxDpbFrames = sps.GetMaxDPBSize();
-				RANGE_CHECK_FAILURE(sps.max_num_ref_frames <= MaxDpbFrames, TEXT("max_num_ref_frames"));
-
+				// Removed for now as this can happen despite it being invalid.
+					// RANGE_CHECK_FAILURE(sps.max_num_ref_frames <= MaxDpbFrames, TEXT("max_num_ref_frames"));
 				bool bIsExt = (sps.profile_idc == 44 || sps.profile_idc == 86 || sps.profile_idc == 100 || sps.profile_idc == 110 || sps.profile_idc == 122 || sps.profile_idc == 244) && sps.constraint_set3_flag!=0;
 				sps.max_num_reorder_frames = sps.max_dec_frame_buffering = bIsExt ? 0 : MaxDpbFrames;
 				sps.vui_parameters_present_flag = br.GetBits(1);
diff --git a/Engine/Plugins/Media/ElectraCodecs/Source/ElectraDecoders/Private/Utils/MPEG/ElectraUtilsMPEGVideo_H265.cpp b/Engine/Plugins/Media/ElectraCodecs/Source/ElectraDecoders/Private/Utils/MPEG/ElectraUtilsMPEGVideo_H265.cpp
--- a/Engine/Plugins/Media/ElectraCodecs/Source/ElectraDecoders/Private/Utils/MPEG/ElectraUtilsMPEGVideo_H265.cpp
+++ b/Engine/Plugins/Media/ElectraCodecs/Source/ElectraDecoders/Private/Utils/MPEG/ElectraUtilsMPEGVideo_H265.cpp
@@ -199,6 +199,10 @@ namespace ElectraDecodersUtil
 				check((int32)max_dec_buffering <= MaxDPB);
 				return (int32)max_dec_buffering;
 			}
+			int32 FSequenceParameterSet::GetMinCbSizeY() const
+			{
+				return MinCbSizeY;
+			}
 			int32 FSequenceParameterSet::GetWidth() const
 			{
 				return (int32) pic_width_in_luma_samples;
@@ -207,6 +211,13 @@ namespace ElectraDecodersUtil
 			{
 				return (int32) pic_height_in_luma_samples;
 			}
+			void FSequenceParameterSet::GetDisplaySize(int32& OutWidth, int32& OutHeight) const
+			{
+				int32 cl,cr,ct,cb;
+				GetCrop(cl,cr,ct,cb);
+				OutWidth = GetWidth() - cl - cr;
+				OutHeight = GetHeight() - ct - cb;
+			}
 			void FSequenceParameterSet::GetCrop(int32& OutLeft, int32& OutRight, int32& OutTop, int32& OutBottom) const
 			{
 				OutLeft = conf_win_left_offset * SubWidthC;
diff --git a/Engine/Plugins/Media/ElectraCodecs/Source/ElectraDecoders/Public/Utils/MPEG/ElectraUtilsMPEGVideo_H264.h b/Engine/Plugins/Media/ElectraCodecs/Source/ElectraDecoders/Public/Utils/MPEG/ElectraUtilsMPEGVideo_H264.h
--- a/Engine/Plugins/Media/ElectraCodecs/Source/ElectraDecoders/Public/Utils/MPEG/ElectraUtilsMPEGVideo_H264.h
+++ b/Engine/Plugins/Media/ElectraCodecs/Source/ElectraDecoders/Public/Utils/MPEG/ElectraUtilsMPEGVideo_H264.h
@@ -181,6 +181,7 @@ namespace ElectraDecodersUtil
 				int32 GetDPBSize() const;
 				int32 GetWidth() const;
 				int32 GetHeight() const;
+				void GetDisplaySize(int32& OutWidth, int32& OutHeight) const;
 				void GetCrop(int32& OutLeft, int32& OutRight, int32& OutTop, int32& OutBottom) const;
 				void GetAspect(int32& OutSarW, int32& OutSarH) const;
 				FFractionalValue GetTiming() const;
diff --git a/Engine/Plugins/Media/ElectraCodecs/Source/ElectraDecoders/Public/Utils/MPEG/ElectraUtilsMPEGVideo_H265.h b/Engine/Plugins/Media/ElectraCodecs/Source/ElectraDecoders/Public/Utils/MPEG/ElectraUtilsMPEGVideo_H265.h
--- a/Engine/Plugins/Media/ElectraCodecs/Source/ElectraDecoders/Public/Utils/MPEG/ElectraUtilsMPEGVideo_H265.h
+++ b/Engine/Plugins/Media/ElectraCodecs/Source/ElectraDecoders/Public/Utils/MPEG/ElectraUtilsMPEGVideo_H265.h
@@ -614,8 +614,10 @@ namespace ElectraDecodersUtil
 				}
 				int32 GetMaxDPBSize() const;
 				int32 GetDPBSize() const;
+				int32 GetMinCbSizeY() const;
 				int32 GetWidth() const;
 				int32 GetHeight() const;
+				void GetDisplaySize(int32& OutWidth, int32& OutHeight) const;
 				void GetCrop(int32& OutLeft, int32& OutRight, int32& OutTop, int32& OutBottom) const;
 				void GetAspect(int32& OutSarW, int32& OutSarH) const;
 				FFractionalValue GetTiming() const;
diff --git a/Engine/Plugins/Media/ElectraPlayer/ElectraPlayer.uplugin b/Engine/Plugins/Media/ElectraPlayer/ElectraPlayer.uplugin
--- a/Engine/Plugins/Media/ElectraPlayer/ElectraPlayer.uplugin
+++ b/Engine/Plugins/Media/ElectraPlayer/ElectraPlayer.uplugin
@@ -49,7 +49,7 @@
 			"PlatformAllowList": [ "Win64", "Mac", "IOS", "TVOS", "Android", "Linux" ],
 			"TargetDenyList": [ "Server" ]
 		}
-		
+
 	],
 	"Plugins": [
 		{
@@ -71,6 +71,10 @@
     {
       "Name": "OpusDecoderElectra",
       "Enabled": true
+    },
+    {
+      "Name": "D3D12VideoDecodersElectra",
+      "Enabled": true
     }
 	]
 }
diff --git a/Engine/Plugins/Media/ElectraPlayer/Source/ElectraPlayerRuntime/Private/Runtime/Player/ABRRules/ABRVoDPlus.cpp b/Engine/Plugins/Media/ElectraPlayer/Source/ElectraPlayerRuntime/Private/Runtime/Player/ABRRules/ABRVoDPlus.cpp
--- a/Engine/Plugins/Media/ElectraPlayer/Source/ElectraPlayerRuntime/Private/Runtime/Player/ABRRules/ABRVoDPlus.cpp
+++ b/Engine/Plugins/Media/ElectraPlayer/Source/ElectraPlayerRuntime/Private/Runtime/Player/ABRRules/ABRVoDPlus.cpp
@@ -18,7 +18,13 @@ static TAutoConsoleVariable<int32> CVarElectraABRVODCycle(
 	TEXT("Electra.ABR.VoD.CycleQualities"),
 	0,
 	TEXT("Cycles through all quality levels even if not sustainable.\n")
-	TEXT(" 0: do not cycle; 1-n: cycle up to quality level n."),
+	TEXT(" 0: do not cycle; 1 to n: cycle up to quality level n."),
+	ECVF_Default);
+static TAutoConsoleVariable<int32> CVarElectraABRVODFixedStream(
+	TEXT("Electra.ABR.VoD.FixedStream"),
+	0,
+	TEXT("Forces use of a particular quality even if not sustainable.\n")
+	TEXT(" 0: do not lock; 1 to n: lock to quality level n-1."),
 	ECVF_Default);
 #endif
 
@@ -71,14 +77,14 @@ private:
 		};
 
 		FStreamWorkVars()
-		{ 
+		{
 			const int32 HistorySize = 5;
 			AverageBandwidth.Resize(HistorySize);
 			const int32 LatencyHistorySize = 5;
 			AverageLatency.Resize(LatencyHistorySize);
-			Reset(); 
+			Reset();
 		}
-		
+
 		void ClearForNextDownload()
 		{
 			QualityIndexDownloading = 0;
@@ -152,7 +158,7 @@ private:
 		const TSharedPtrTS<FABRStreamInformation>* Stream = StreamInfos.FindByPredicate([InQualityIndex](const TSharedPtrTS<FABRStreamInformation>& InInfo) { return InQualityIndex == InInfo->QualityIndex;} );
 		return Stream ? (*Stream) : nullptr;
 	}
-	
+
 	double GetPlayablePlayerDuration(bool& bEOS, EStreamType InStreamType)
 	{
 		IAdaptiveStreamSelector::IPlayerLiveControl::FABRBufferStats bs;
@@ -242,7 +248,7 @@ private:
 
 	// Initially assumed network latency.
 	const double DefaultNetworkLatency = 0.4;
-	
+
 	// Scale of highest stream bitrate to clamp bandwidth to so it does not get ridiculously large.
 	const double ClampBandwidthToMaxStreamBitrateScaleFactor = 2.0;
 	// Clamp to this rate unless the max stream bitrate after applying the scaling factor is already higher.
@@ -252,6 +258,8 @@ private:
 	bool bForceQualityCycling = false;
 	int32 QualityCyclingIndex = 0;
 	int32 QualityCyclingLimit = 0;
+	bool bForceQualityLock = false;
+	int32 QualityLockIndex = 0;
 #endif
 };
 
@@ -926,6 +934,15 @@ IAdaptiveStreamSelector::ESegmentAction FABROnDemandPlus::PerformSelection(const
 					Info->LogMessage(IInfoLog::ELevel::Info, FString::Printf(TEXT("Cycling to quality index %d (%d*%d @ %d bps)"), NewQualityIndex, Cnd->Resolution.Width, Cnd->Resolution.Height, Cnd->Bitrate));
 				}
 			}
+			if (bForceQualityLock)
+			{
+				NewQualityIndex = Utils::Min(InCandidates.Num(), QualityLockIndex) - 1;
+				TSharedPtrTS<FABRStreamInformation> Cnd = GetStreamInfoForQualityIndex(StreamType, NewQualityIndex);
+				if (Cnd.IsValid())
+				{
+					Info->LogMessage(IInfoLog::ELevel::Info, FString::Printf(TEXT("Forcing quality index %d (%d*%d @ %d bps)"), NewQualityIndex, Cnd->Resolution.Width, Cnd->Resolution.Height, Cnd->Bitrate));
+				}
+			}
 #endif
 		}
 		else
@@ -974,7 +991,7 @@ IAdaptiveStreamSelector::ESegmentAction FABROnDemandPlus::PerformSelection(const
 
 
 FTimeValue FABROnDemandPlus::GetMinBufferTimeForPlayback(IAdaptiveStreamSelector::EMinBufferType InBufferingType, FTimeValue InDefaultMBT)
-{ 
+{
 	return FTimeValue();
 }
 
@@ -1032,6 +1049,20 @@ IAdaptiveStreamSelector::EHandlingAction FABROnDemandPlus::PeriodicHandle()
 		QualityCyclingLimit = 0;
 		Info->LogMessage(IInfoLog::ELevel::Info, FString::Printf(TEXT("Diable quality cycling")));
 	}
+
+	int32 LockQuality = CVarElectraABRVODFixedStream.GetValueOnAnyThread();
+	if (LockQuality && LockQuality!=QualityLockIndex)
+	{
+		bForceQualityLock = true;
+		QualityLockIndex = LockQuality;
+		Info->LogMessage(IInfoLog::ELevel::Info, FString::Printf(TEXT("Locking quality to index %d"), LockQuality));
+	}
+	else if (LockQuality == 0 && bForceQualityLock)
+	{
+		bForceQualityLock = false;
+		QualityLockIndex = 0;
+		Info->LogMessage(IInfoLog::ELevel::Info, FString::Printf(TEXT("Diable quality locking")));
+	}
 #endif
 
 	return NextAction;
