D3D12VideoDecoderElectra: removed preprocessor error tag

#rnx

diff --git a/Engine/Plugins/Media/D3D12VideoDecodersElectra/Source/Private/VideoDecoder_D3D12.cpp b/Engine/Plugins/Media/D3D12VideoDecodersElectra/Source/Private/VideoDecoder_D3D12.cpp
--- a/Engine/Plugins/Media/D3D12VideoDecodersElectra/Source/Private/VideoDecoder_D3D12.cpp
+++ b/Engine/Plugins/Media/D3D12VideoDecodersElectra/Source/Private/VideoDecoder_D3D12.cpp
@@ -23,11 +23,9 @@
 namespace ElectraVideoDecodersD3D12Video
 {
 #ifndef ELECTRA_DECODERS_D3D12VIDEO_DISABLED_ON_PLATFORM
-#error "C"
 #define ELECTRA_DECODERS_D3D12VIDEO_DISABLED_ON_PLATFORM 1
 #endif
 #ifndef ELECTRA_DECODERS_D3D12VIDEO_IGNORED_ON_PLATFORM
-#error "B"
 #define ELECTRA_DECODERS_D3D12VIDEO_IGNORED_ON_PLATFORM 1
 #endif
 
