ElectraDecoders: Added missing error handling when decoder could not be created, resulting in the use of member variable that was nullptr then.

#rnx
#tests Tested that issue occurs when decoder creation is made to fail on purpose, then applied fix to see that it now fails gracefully.

diff --git a/Engine/Plugins/Media/D3D12VideoDecodersElectra/Source/Private/VideoDecoder_D3D12_H264.cpp b/Engine/Plugins/Media/D3D12VideoDecodersElectra/Source/Private/VideoDecoder_D3D12_H264.cpp
--- a/Engine/Plugins/Media/D3D12VideoDecodersElectra/Source/Private/VideoDecoder_D3D12_H264.cpp
+++ b/Engine/Plugins/Media/D3D12VideoDecodersElectra/Source/Private/VideoDecoder_D3D12_H264.cpp
@@ -198,7 +198,10 @@ IElectraDecoder::EDecoderError FD3D12VideoDecoder_H264::DecodeAccessUnit(const F
 		// Create a new decoder if we do not have one. This does not require any information about the resolution or DPB.
 		if (!VideoDecoder.IsValid())
 		{
-			InternalDecoderCreate();
+			if (!InternalDecoderCreate())
+			{
+				return IElectraDecoder::EDecoderError::Error;
+			}
 		}
 
 		const ElectraDecodersUtil::MPEG::H264::FPictureParameterSet* ppsPtr = BitstreamParamsH264.PPSs.Find(SliceInfos[0].Header.pic_parameter_set_id);
@@ -301,6 +304,10 @@ IElectraDecoder::EDecoderError FD3D12VideoDecoder_H264::DecodeSlicesH264(const F
 	// The caller needs to make sure we do not get called without slices
 	check(InSliceInfos.Num());
 
+	if (!VideoDecoderSync.IsValid())
+	{
+		return IElectraDecoder::EDecoderError::Error;
+	}
 	// The previous operation must have completed, primarily because we (may) need the decoded frame from before
 	// as a reference frame for this call and that frame thus needs to have finished.
 	if (!VideoDecoderSync->AwaitCompletion(500))
diff --git a/Engine/Plugins/Media/D3D12VideoDecodersElectra/Source/Private/VideoDecoder_D3D12_H265.cpp b/Engine/Plugins/Media/D3D12VideoDecodersElectra/Source/Private/VideoDecoder_D3D12_H265.cpp
--- a/Engine/Plugins/Media/D3D12VideoDecodersElectra/Source/Private/VideoDecoder_D3D12_H265.cpp
+++ b/Engine/Plugins/Media/D3D12VideoDecodersElectra/Source/Private/VideoDecoder_D3D12_H265.cpp
@@ -227,7 +227,10 @@ IElectraDecoder::EDecoderError FD3D12VideoDecoder_H265::DecodeAccessUnit(const F
 		// Create a new decoder if we do not have one. This does not require any information about the resolution or DPB.
 		if (!VideoDecoder.IsValid())
 		{
-			InternalDecoderCreate();
+			if (!InternalDecoderCreate())
+			{
+				return IElectraDecoder::EDecoderError::Error;
+			}
 		}
 
 		const ElectraDecodersUtil::MPEG::H265::FPictureParameterSet* ppsPtr = BitstreamParamsH265.PPSs.Find(SliceInfos[0].Header.slice_pic_parameter_set_id);
@@ -344,6 +347,10 @@ IElectraDecoder::EDecoderError FD3D12VideoDecoder_H265::DecodeSlicesH265(const F
 	// The caller needs to make sure we do not get called without slices
 	check(InSliceInfos.Num());
 
+	if (!VideoDecoderSync.IsValid())
+	{
+		return IElectraDecoder::EDecoderError::Error;
+	}
 	// The previous operation must have completed, primarily because we (may) need the decoded frame from before
 	// as a reference frame for this call and that frame thus needs to have finished.
 	if (!VideoDecoderSync->AwaitCompletion(500))
