ElectraCodecs: Draining the D3D12 video decoder on resolution changes unconditionally since there are too many parameters to track and compare to check if the streams are compatible.

#rnx

diff --git a/Engine/Plugins/Media/D3D12VideoDecodersElectra/Source/Private/VideoDecoder_D3D12_H264.cpp b/Engine/Plugins/Media/D3D12VideoDecodersElectra/Source/Private/VideoDecoder_D3D12_H264.cpp
--- a/Engine/Plugins/Media/D3D12VideoDecodersElectra/Source/Private/VideoDecoder_D3D12_H264.cpp
+++ b/Engine/Plugins/Media/D3D12VideoDecodersElectra/Source/Private/VideoDecoder_D3D12_H264.cpp
@@ -38,6 +38,7 @@ IElectraDecoder::ECSDCompatibility FD3D12VideoDecoder_H264::IsCompatibleWith(con
 		return IElectraDecoder::ECSDCompatibility::DrainAndReset;
 	}
 
+#if 0
 	const ElectraDecodersUtil::MPEG::H264::FSequenceParameterSet sps = temp.SPSs.CreateConstIterator().Value();
 	int32 NewWidth, NewHeight, NewDPBSize;
 	NewDPBSize = sps.GetDPBSize();
@@ -49,6 +50,9 @@ IElectraDecoder::ECSDCompatibility FD3D12VideoDecoder_H264::IsCompatibleWith(con
 		return IElectraDecoder::ECSDCompatibility::DrainAndReset;
 	}
 	return IElectraDecoder::ECSDCompatibility::Compatible;
+#else
+	return IElectraDecoder::ECSDCompatibility::DrainAndReset;
+#endif
 }
 
 TSharedPtr<IElectraDecoderDefaultOutputFormat, ESPMode::ThreadSafe> FD3D12VideoDecoder_H264::GetDefaultOutputFormatFromCSD(const TMap<FString, FVariant>& CSDAndAdditionalOptions)
diff --git a/Engine/Plugins/Media/D3D12VideoDecodersElectra/Source/Private/VideoDecoder_D3D12_H265.cpp b/Engine/Plugins/Media/D3D12VideoDecodersElectra/Source/Private/VideoDecoder_D3D12_H265.cpp
--- a/Engine/Plugins/Media/D3D12VideoDecodersElectra/Source/Private/VideoDecoder_D3D12_H265.cpp
+++ b/Engine/Plugins/Media/D3D12VideoDecodersElectra/Source/Private/VideoDecoder_D3D12_H265.cpp
@@ -44,6 +44,7 @@ IElectraDecoder::ECSDCompatibility FD3D12VideoDecoder_H265::IsCompatibleWith(con
 		return IElectraDecoder::ECSDCompatibility::DrainAndReset;
 	}
 
+#if 0
 	const ElectraDecodersUtil::MPEG::H265::FSequenceParameterSet sps = temp.SPSs.CreateConstIterator().Value();
 	// Check that the new CSD isn't Main10 when we are only Main.
 	if (sps.profile_tier_level.general_profile_idc == 2 && !CodecInfo.b10Bit)
@@ -60,6 +61,9 @@ IElectraDecoder::ECSDCompatibility FD3D12VideoDecoder_H265::IsCompatibleWith(con
 		return IElectraDecoder::ECSDCompatibility::DrainAndReset;
 	}
 	return IElectraDecoder::ECSDCompatibility::Compatible;
+#else
+	return IElectraDecoder::ECSDCompatibility::DrainAndReset;
+#endif
 }
 
 TSharedPtr<IElectraDecoderDefaultOutputFormat, ESPMode::ThreadSafe> FD3D12VideoDecoder_H265::GetDefaultOutputFormatFromCSD(const TMap<FString, FVariant>& CSDAndAdditionalOptions)
diff --git a/Engine/Plugins/Media/ElectraCodecs/Source/ElectraDecoders/Private/Utils/MPEG/ElectraUtilsMPEGVideo_H264.cpp b/Engine/Plugins/Media/ElectraCodecs/Source/ElectraDecoders/Private/Utils/MPEG/ElectraUtilsMPEGVideo_H264.cpp
--- a/Engine/Plugins/Media/ElectraCodecs/Source/ElectraDecoders/Private/Utils/MPEG/ElectraUtilsMPEGVideo_H264.cpp
+++ b/Engine/Plugins/Media/ElectraCodecs/Source/ElectraDecoders/Private/Utils/MPEG/ElectraUtilsMPEGVideo_H264.cpp
@@ -1228,7 +1228,7 @@ namespace ElectraDecodersUtil
 				uint32 UnusedShortTermFrameNum = (CurrentPOC.prev_frame_num + 1) % max_frame_num;
 				FOutputFrameInfo NoInfo;
 				NoInfo.bDoNotOutput = true;
-				while(InSliceHeader.frame_num != UnusedShortTermFrameNum)
+				for(int32 nm=0; nm<max_frame_num && InSliceHeader.frame_num != UnusedShortTermFrameNum; ++nm)
 				{
 					TempSliceHeader.frame_num = UnusedShortTermFrameNum;
 					if (InSequenceParameterSet.pic_order_cnt_type != 0)
