Disable D3D12 decoder when reference only allocation flag is encountered
#rnx

diff --git a/Engine/Plugins/Media/D3D12VideoDecodersElectra/Source/Private/VideoDecoder_D3D12.cpp b/Engine/Plugins/Media/D3D12VideoDecodersElectra/Source/Private/VideoDecoder_D3D12.cpp
--- a/Engine/Plugins/Media/D3D12VideoDecodersElectra/Source/Private/VideoDecoder_D3D12.cpp
+++ b/Engine/Plugins/Media/D3D12VideoDecodersElectra/Source/Private/VideoDecoder_D3D12.cpp
@@ -416,6 +416,13 @@ public:
 					Codec = nullptr;
 				}
 
+				// We don't support reference only allocations yet
+				if ((OutSupport.ConfigurationFlags & D3D12_VIDEO_DECODE_CONFIGURATION_FLAG_REFERENCE_ONLY_ALLOCATIONS_REQUIRED) != 0)
+				{
+					UE_LOG(LogD3D12VideoDecodersElectra, Log, TEXT("Decode reference only allocations are not supported, but configuration flags 0x%08x were returned."), (int)OutSupport.ConfigurationFlags);
+					Codec = nullptr;
+				}
+
 				// Do a custom platform capability check.
 				if (!FD3D12VideoDecoder::CheckPlatformDecodeCapabilities(OutSupport, ci, InOptions))
 				{
