ElectraDecoders: Adjuted H.265 crop values after frame alignment in D3D12 decoder

#rnx

diff --git a/Engine/Plugins/Media/D3D12VideoDecodersElectra/Source/Private/VideoDecoder_D3D12_H265.cpp b/Engine/Plugins/Media/D3D12VideoDecodersElectra/Source/Private/VideoDecoder_D3D12_H265.cpp
--- a/Engine/Plugins/Media/D3D12VideoDecodersElectra/Source/Private/VideoDecoder_D3D12_H265.cpp
+++ b/Engine/Plugins/Media/D3D12VideoDecodersElectra/Source/Private/VideoDecoder_D3D12_H265.cpp
@@ -770,6 +770,9 @@ IElectraDecoder::EDecoderError FD3D12VideoDecoder_H265::DecodeSlicesH265(const F
 
 	InDec->Width = Align(InSequenceParameterSet.GetWidth(), GetFrameAlignment());
 	InDec->Height = Align(InSequenceParameterSet.GetHeight(), GetFrameAlignment());
+	// Adjust the cropping values to the right and bottom to include the required alignment we had to add.
+	InDec->Crop.Right += InDec->Width - InSequenceParameterSet.GetWidth();
+	InDec->Crop.Bottom += InDec->Height - InSequenceParameterSet.GetHeight();
 	InDec->ImageWidth = InDec->Width - InDec->Crop.Left - InDec->Crop.Right;
 	InDec->ImageHeight = InDec->Height - InDec->Crop.Top - InDec->Crop.Bottom;
 	//InDec->Pitch = InDec->ImageWidth;
