ElectraCodecs: using format specific bitstream parser in D3D12 decoders to pass HDR metadata

#rnx

diff --git a/Engine/Plugins/Media/D3D12VideoDecodersElectra/Source/Private/VideoDecoder_D3D12.cpp b/Engine/Plugins/Media/D3D12VideoDecodersElectra/Source/Private/VideoDecoder_D3D12.cpp
--- a/Engine/Plugins/Media/D3D12VideoDecodersElectra/Source/Private/VideoDecoder_D3D12.cpp
+++ b/Engine/Plugins/Media/D3D12VideoDecodersElectra/Source/Private/VideoDecoder_D3D12.cpp
@@ -579,6 +579,7 @@ public:
 FD3D12VideoDecoder::FD3D12VideoDecoder(const FCodecFormatHelper::FCodecInfo& InCodecInfo, const D3D12_FEATURE_DATA_VIDEO_DECODE_SUPPORT& InDecodeSupport, const TMap<FString, FVariant>& InOptions, TSharedPtr<IElectraDecoderResourceDelegate, ESPMode::ThreadSafe> InResourceDelegate, const TRefCountPtr<ID3D12Device>& InD3D12Device, const TRefCountPtr<ID3D12VideoDevice>& InVideoDevice, uint32 InVideoDeviceNodeIndex)
 	: CodecInfo(InCodecInfo)
 	, DecodeSupport(InDecodeSupport)
+	, InitialCreationOptions(InOptions)
 	, ResourceDelegate(InResourceDelegate)
 	, D3D12Device(InD3D12Device)
 	, VideoDevice(InVideoDevice)
diff --git a/Engine/Plugins/Media/D3D12VideoDecodersElectra/Source/Private/VideoDecoder_D3D12_Common.h b/Engine/Plugins/Media/D3D12VideoDecodersElectra/Source/Private/VideoDecoder_D3D12_Common.h
--- a/Engine/Plugins/Media/D3D12VideoDecodersElectra/Source/Private/VideoDecoder_D3D12_Common.h
+++ b/Engine/Plugins/Media/D3D12VideoDecodersElectra/Source/Private/VideoDecoder_D3D12_Common.h
@@ -13,7 +13,6 @@
 #include "IElectraDecoderOutputVideo.h"
 
 #include "ElectraDecodersUtils.h"
-#include "Utils/ElectraBitstreamProcessorDefault.h"
 
 /*********************************************************************************************************************/
 #include COMPILED_PLATFORM_HEADER(PlatformHeaders_Video_D3D.h)
@@ -358,8 +357,7 @@ protected:
 	EDecoderError Flush() = 0;
 	EOutputStatus HaveOutput() override;
 	TSharedPtr<IElectraDecoderOutput, ESPMode::ThreadSafe> GetOutput() override;
-	TSharedPtr<IElectraDecoderBitstreamProcessor, ESPMode::ThreadSafe> CreateBitstreamProcessor() override
-	{ return FElectraDecoderBitstreamProcessorDefault::Create(); }
+	TSharedPtr<IElectraDecoderBitstreamProcessor, ESPMode::ThreadSafe> CreateBitstreamProcessor() = 0;
 	void Suspend() override
 	{ }
 	void Resume() override
@@ -471,6 +469,7 @@ protected:
 
 	FCodecFormatHelper::FCodecInfo CodecInfo;
 	D3D12_FEATURE_DATA_VIDEO_DECODE_SUPPORT DecodeSupport;
+	TMap<FString, FVariant> InitialCreationOptions;
 	TWeakPtr<IElectraDecoderResourceDelegate, ESPMode::ThreadSafe> ResourceDelegate;
 	TRefCountPtr<ID3D12Device> D3D12Device;
 	TRefCountPtr<ID3D12VideoDevice> VideoDevice;
diff --git a/Engine/Plugins/Media/D3D12VideoDecodersElectra/Source/Private/VideoDecoder_D3D12_H264.cpp b/Engine/Plugins/Media/D3D12VideoDecodersElectra/Source/Private/VideoDecoder_D3D12_H264.cpp
--- a/Engine/Plugins/Media/D3D12VideoDecodersElectra/Source/Private/VideoDecoder_D3D12_H264.cpp
+++ b/Engine/Plugins/Media/D3D12VideoDecodersElectra/Source/Private/VideoDecoder_D3D12_H264.cpp
@@ -2,6 +2,7 @@
 
 #include "VideoDecoder_D3D12_H264.h"
 #include "D3D12VideoDecodersElectraModule.h"
+#include "Utils/MPEG/ElectraBitstreamProcessor_H264.h"
 
 
 namespace ElectraVideoDecodersD3D12Video
@@ -60,6 +61,13 @@ TSharedPtr<IElectraDecoderDefaultOutputFormat, ESPMode::ThreadSafe> FD3D12VideoD
 	return nullptr;
 }
 
+TSharedPtr<IElectraDecoderBitstreamProcessor, ESPMode::ThreadSafe> FD3D12VideoDecoder_H264::CreateBitstreamProcessor()
+{
+	TMap<FString, FVariant> DecoderFeatures;
+	GetFeatures(DecoderFeatures);
+	return FElectraDecoderBitstreamProcessorH264::Create(DecoderFeatures, InitialCreationOptions);
+}
+
 
 IElectraDecoder::EDecoderError FD3D12VideoDecoder_H264::GetCodecSpecificDataH264(FBitstreamParamsH264& OutBitstreamParamsH264, const TMap<FString, FVariant>& InAdditionalOptions, bool bIsRequired)
 {
diff --git a/Engine/Plugins/Media/D3D12VideoDecodersElectra/Source/Private/VideoDecoder_D3D12_H264.h b/Engine/Plugins/Media/D3D12VideoDecodersElectra/Source/Private/VideoDecoder_D3D12_H264.h
--- a/Engine/Plugins/Media/D3D12VideoDecodersElectra/Source/Private/VideoDecoder_D3D12_H264.h
+++ b/Engine/Plugins/Media/D3D12VideoDecodersElectra/Source/Private/VideoDecoder_D3D12_H264.h
@@ -18,6 +18,7 @@ public:
 protected:
 	ECSDCompatibility IsCompatibleWith(const TMap<FString, FVariant>& CSDAndAdditionalOptions) override;
 	TSharedPtr<IElectraDecoderDefaultOutputFormat, ESPMode::ThreadSafe> GetDefaultOutputFormatFromCSD(const TMap<FString, FVariant>& CSDAndAdditionalOptions) override;
+	TSharedPtr<IElectraDecoderBitstreamProcessor, ESPMode::ThreadSafe> CreateBitstreamProcessor() override;
 	EDecoderError DecodeAccessUnit(const FInputAccessUnit& InInputAccessUnit, const TMap<FString, FVariant>& InAdditionalOptions) override;
 	EDecoderError SendEndOfData() override;
 	EDecoderError Flush() override;
diff --git a/Engine/Plugins/Media/D3D12VideoDecodersElectra/Source/Private/VideoDecoder_D3D12_H265.cpp b/Engine/Plugins/Media/D3D12VideoDecodersElectra/Source/Private/VideoDecoder_D3D12_H265.cpp
--- a/Engine/Plugins/Media/D3D12VideoDecodersElectra/Source/Private/VideoDecoder_D3D12_H265.cpp
+++ b/Engine/Plugins/Media/D3D12VideoDecodersElectra/Source/Private/VideoDecoder_D3D12_H265.cpp
@@ -2,6 +2,7 @@
 
 #include "VideoDecoder_D3D12_H265.h"
 #include "D3D12VideoDecodersElectraModule.h"
+#include "Utils/MPEG/ElectraBitstreamProcessor_H265.h"
 
 
 namespace ElectraVideoDecodersD3D12Video
@@ -71,6 +72,13 @@ TSharedPtr<IElectraDecoderDefaultOutputFormat, ESPMode::ThreadSafe> FD3D12VideoD
 	return nullptr;
 }
 
+TSharedPtr<IElectraDecoderBitstreamProcessor, ESPMode::ThreadSafe> FD3D12VideoDecoder_H265::CreateBitstreamProcessor()
+{
+	TMap<FString, FVariant> DecoderFeatures;
+	GetFeatures(DecoderFeatures);
+	return FElectraDecoderBitstreamProcessorH265::Create(DecoderFeatures, InitialCreationOptions);
+}
+
 
 IElectraDecoder::EDecoderError FD3D12VideoDecoder_H265::GetCodecSpecificDataH265(FBitstreamParamsH265& OutBitstreamParamsH265, const TMap<FString, FVariant>& InAdditionalOptions, bool bIsRequired)
 {
diff --git a/Engine/Plugins/Media/D3D12VideoDecodersElectra/Source/Private/VideoDecoder_D3D12_H265.h b/Engine/Plugins/Media/D3D12VideoDecodersElectra/Source/Private/VideoDecoder_D3D12_H265.h
--- a/Engine/Plugins/Media/D3D12VideoDecodersElectra/Source/Private/VideoDecoder_D3D12_H265.h
+++ b/Engine/Plugins/Media/D3D12VideoDecodersElectra/Source/Private/VideoDecoder_D3D12_H265.h
@@ -18,6 +18,7 @@ public:
 protected:
 	ECSDCompatibility IsCompatibleWith(const TMap<FString, FVariant>& CSDAndAdditionalOptions) override;
 	TSharedPtr<IElectraDecoderDefaultOutputFormat, ESPMode::ThreadSafe> GetDefaultOutputFormatFromCSD(const TMap<FString, FVariant>& CSDAndAdditionalOptions) override;
+	TSharedPtr<IElectraDecoderBitstreamProcessor, ESPMode::ThreadSafe> CreateBitstreamProcessor() override;
 	EDecoderError DecodeAccessUnit(const FInputAccessUnit& InInputAccessUnit, const TMap<FString, FVariant>& InAdditionalOptions) override;
 	EDecoderError SendEndOfData() override;
 	EDecoderError Flush() override;
