Electra: D3D12 HEVC video decoding related changes

#rnx

#Synced-CL 51578981

diff --git a/Engine/Plugins/Media/D3D12VideoDecodersElectra/Source/Private/VideoDecoder_D3D12.cpp b/Engine/Plugins/Media/D3D12VideoDecodersElectra/Source/Private/VideoDecoder_D3D12.cpp
--- a/Engine/Plugins/Media/D3D12VideoDecodersElectra/Source/Private/VideoDecoder_D3D12.cpp
+++ b/Engine/Plugins/Media/D3D12VideoDecodersElectra/Source/Private/VideoDecoder_D3D12.cpp
@@ -1013,6 +1013,8 @@ bool FD3D12VideoDecoder::CreateDPB(TSharedPtr<FDecodedPictureBuffer, ESPMode::Th
 bool FD3D12VideoDecoder::PrepareBitstreamBuffer(const TSharedPtr<FFrameDecodeResource, ESPMode::ThreadSafe>& InFrameDecodeResourceToPrepare, uint32 InMaxInputBufferSize)
 {
 	check(InFrameDecodeResourceToPrepare);
+	// Align to 256KiB multiples to reduce the number of times this needs to be marginally enlarged.
+	InMaxInputBufferSize = Align(InMaxInputBufferSize, 1024*256);
 	if (InFrameDecodeResourceToPrepare && (!InFrameDecodeResourceToPrepare->D3DBitstreamBuffer.IsValid() || InFrameDecodeResourceToPrepare->D3DBitstreamBufferAllocatedSize < InMaxInputBufferSize))
 	{
 		InFrameDecodeResourceToPrepare->D3DBitstreamBuffer.SafeRelease();
diff --git a/Engine/Plugins/Media/D3D12VideoDecodersElectra/Source/Private/VideoDecoder_D3D12_H265.cpp b/Engine/Plugins/Media/D3D12VideoDecodersElectra/Source/Private/VideoDecoder_D3D12_H265.cpp
--- a/Engine/Plugins/Media/D3D12VideoDecodersElectra/Source/Private/VideoDecoder_D3D12_H265.cpp
+++ b/Engine/Plugins/Media/D3D12VideoDecodersElectra/Source/Private/VideoDecoder_D3D12_H265.cpp
@@ -286,7 +286,7 @@ IElectraDecoder::EDecoderError FD3D12VideoDecoder_H265::DecodeAccessUnit(const F
 				// for this stream (the largest resolution).
 				const int32 Width = (int32) DecodeSupport.Width;
 				const int32 Height = (int32) DecodeSupport.Height;
-				const int32 NumFrames = spsPtr->GetDPBSize() + 2;	// 1 extra for the current frame that's not in the DPB yet, and 1 extra that acts as a 'missing' frame.
+				const int32 NumFrames = DPBSize + 2;	// 1 extra for the current frame that's not in the DPB yet, and 1 extra that acts as a 'missing' frame.
 				if (!CreateDPB(DPB, Width, Height, GetFrameAlignment(), NumFrames))
 				{
 					return IElectraDecoder::EDecoderError::Error;
@@ -298,7 +298,34 @@ IElectraDecoder::EDecoderError FD3D12VideoDecoder_H265::DecodeAccessUnit(const F
 					return IElectraDecoder::EDecoderError::Error;
 				}
 			}
+			// Remember the type of IRAP (see below).
+			BitstreamParamsH265.AssociatedIRAP = SliceInfos[0].NalUnitType;
+			++BitstreamParamsH265.IRAPCount;
 		}
+		/*
+			Filter out any RASL pictures for which the associated IRAP, that is the first coded picture, is a CRA or BLA.
+
+			We do this because the DXVA HEVC spec states:
+				"The host decoder is recommended to send only decodable pictures to accelerator. For example, when the associated IRAP picture
+				is a BLA picture or is a CRA picture that is the first coded picture in the bitstream, the host decoder is recommended not to
+				send the associated RASL pictures to accelerator if any such RASL pictures are present in the bitstream. 
+				However, the accelerator shall be robust enough to handle any non-decodable pictures."
+
+			And the "However" statement does not hold true. There are accelerators where sending down these pictures will produce
+			garbage frames. So we need to filter them out.
+		*/
+		if (BitstreamParamsH265.IRAPCount <= 1 && (BitstreamParamsH265.AssociatedIRAP == 21 || BitstreamParamsH265.AssociatedIRAP == 16 || BitstreamParamsH265.AssociatedIRAP == 17 || BitstreamParamsH265.AssociatedIRAP == 18))
+		{
+			for(int32 i=0; i<SliceInfos.Num(); ++i)
+			{
+				if (SliceInfos[i].Header.NalUnitType == 8 || SliceInfos[i].Header.NalUnitType == 9)
+				{
+					SliceInfos.RemoveAt(i);
+					--i;
+				}
+			}
+		}
+
 		IElectraDecoder::EDecoderError Error = DecodeSlicesH265(InInputAccessUnit, SliceInfos, *spsPtr, *ppsPtr);
 		// When the frame was decoded and there is an EOS/EOB NALU present we need to apply it now.
 		if (Error == IElectraDecoder::EDecoderError::None && SequenceEnd == ESequenceEnd::AtEnd)
@@ -326,6 +353,8 @@ IElectraDecoder::EDecoderError FD3D12VideoDecoder_H265::SendEndOfData()
 	TArray<ElectraDecodersUtil::MPEG::H265::FDecodedPictureBuffer::FDPBOutputFrame> OutputFrames;
 	BitstreamParamsH265.DPB.Flush(OutputFrames);
 	BitstreamParamsH265.bIsFirstInSequence = true;
+	BitstreamParamsH265.AssociatedIRAP = -1;
+	BitstreamParamsH265.IRAPCount = 0;
 	return HandleOutputListH265(OutputFrames);
 }
 
@@ -348,8 +377,11 @@ IElectraDecoder::EDecoderError FD3D12VideoDecoder_H265::Flush()
 
 IElectraDecoder::EDecoderError FD3D12VideoDecoder_H265::DecodeSlicesH265(const FInputAccessUnit& InInputAccessUnit, const TArray<FBitstreamParamsH265::FSliceDecodeInfo>& InSliceInfos, const ElectraDecodersUtil::MPEG::H265::FSequenceParameterSet& InSequenceParameterSet, const ElectraDecodersUtil::MPEG::H265::FPictureParameterSet& InPictureParameterSet)
 {
-	// The caller needs to make sure we do not get called without slices
-	check(InSliceInfos.Num());
+	// When there are no slices we are done.
+	if (InSliceInfos.IsEmpty())
+	{
+		return IElectraDecoder::EDecoderError::None;
+	}
 
 	if (!VideoDecoderSync.IsValid())
 	{
@@ -474,14 +506,14 @@ IElectraDecoder::EDecoderError FD3D12VideoDecoder_H265::DecodeSlicesH265(const F
 	{
 		uint32 SliceSize = 3 + InSliceInfos[ns].NumBytesInSlice;
 		input.SliceHeaders[ns].BSNALunitDataLocation = BufferPtr - BufferBase;
-		input.SliceHeaders[ns].SliceBytesInBuffer = Align(SliceSize, 128);
+		input.SliceHeaders[ns].SliceBytesInBuffer = SliceSize;
 		input.SliceHeaders[ns].wBadSliceChopping = 0;
 		*BufferPtr++ = 0;
 		*BufferPtr++ = 0;
 		*BufferPtr++ = 1;
 		FMemory::Memcpy(BufferPtr, InSliceInfos[ns].NalUnitStartAddress, InSliceInfos[ns].NumBytesInSlice);
 		BufferPtr += InSliceInfos[ns].NumBytesInSlice;
-		if (SliceSize != input.SliceHeaders[ns].SliceBytesInBuffer)
+		if (SliceSize != Align(SliceSize, 128))
 		{
 			uint32 PaddingSize = 128 - (SliceSize & 127);
 			FMemory::Memzero(BufferPtr, PaddingSize);
@@ -557,7 +589,6 @@ IElectraDecoder::EDecoderError FD3D12VideoDecoder_H265::DecodeSlicesH265(const F
 	pp.tiles_enabled_flag = InPictureParameterSet.tiles_enabled_flag;
 	pp.entropy_coding_sync_enabled_flag = InPictureParameterSet.entropy_coding_sync_enabled_flag;
 	pp.uniform_spacing_flag = InPictureParameterSet.uniform_spacing_flag;
-	pp.loop_filter_across_tiles_enabled_flag = InPictureParameterSet.loop_filter_across_tiles_enabled_flag;
 	pp.pps_loop_filter_across_slices_enabled_flag = InPictureParameterSet.pps_loop_filter_across_slices_enabled_flag;
 	pp.deblocking_filter_override_enabled_flag = InPictureParameterSet.deblocking_filter_override_enabled_flag;
 	pp.pps_deblocking_filter_disabled_flag = InPictureParameterSet.pps_deblocking_filter_disabled_flag;
@@ -570,6 +601,7 @@ IElectraDecoder::EDecoderError FD3D12VideoDecoder_H265::DecodeSlicesH265(const F
 	pp.pps_cr_qp_offset = (CHAR) InPictureParameterSet.pps_cr_qp_offset;
 	if (InPictureParameterSet.tiles_enabled_flag)
 	{
+		pp.loop_filter_across_tiles_enabled_flag = InPictureParameterSet.loop_filter_across_tiles_enabled_flag;
 		// CAUTION: The maximum number of tiles in the structure is set to accommodate
 		//          at most level 6.3 with 20x22 tiles. Level 7 and higher allows for 40x44 !!
 		check(InPictureParameterSet.num_tile_columns_minus1 < 20);
diff --git a/Engine/Plugins/Media/D3D12VideoDecodersElectra/Source/Private/VideoDecoder_D3D12_H265.h b/Engine/Plugins/Media/D3D12VideoDecodersElectra/Source/Private/VideoDecoder_D3D12_H265.h
--- a/Engine/Plugins/Media/D3D12VideoDecodersElectra/Source/Private/VideoDecoder_D3D12_H265.h
+++ b/Engine/Plugins/Media/D3D12VideoDecodersElectra/Source/Private/VideoDecoder_D3D12_H265.h
@@ -34,6 +34,8 @@ private:
 			PPSs.Empty();
 			CurrentCSD.Empty();
 			DPB.Reset();
+			AssociatedIRAP = -1;
+			IRAPCount = 0;
 			bIsFirstInSequence = true;
 		}
 		struct FSliceDecodeInfo
@@ -55,6 +57,8 @@ private:
 		TMap<uint32, ElectraDecodersUtil::MPEG::H265::FPictureParameterSet> PPSs;
 		TArray<uint8> CurrentCSD;
 		ElectraDecodersUtil::MPEG::H265::FDecodedPictureBuffer DPB;
+		int32 AssociatedIRAP = -1;
+		int32 IRAPCount = 0;
 		bool bIsFirstInSequence = true;
 	};
 
diff --git a/Engine/Plugins/Media/ElectraCodecs/Source/ElectraDecoders/Private/Utils/MPEG/ElectraUtilsMPEGVideo_H265.cpp b/Engine/Plugins/Media/ElectraCodecs/Source/ElectraDecoders/Private/Utils/MPEG/ElectraUtilsMPEGVideo_H265.cpp
--- a/Engine/Plugins/Media/ElectraCodecs/Source/ElectraDecoders/Private/Utils/MPEG/ElectraUtilsMPEGVideo_H265.cpp
+++ b/Engine/Plugins/Media/ElectraCodecs/Source/ElectraDecoders/Private/Utils/MPEG/ElectraUtilsMPEGVideo_H265.cpp
@@ -21,7 +21,7 @@ namespace ElectraDecodersUtil
 				uint32 MaxLumaPs = 0;
 				switch(InGeneralLevelIdc)
 				{
-					#define LEVEL(a,b) a*30+b
+					#define LEVEL(a,b) a*30+(b*3)
 					case LEVEL(1,0):	MaxLumaPs = 36864;		break;
 					case LEVEL(2,0):	MaxLumaPs = 122880;		break;
 					case LEVEL(2,1):	MaxLumaPs = 245760;		break;
@@ -333,7 +333,7 @@ namespace ElectraDecodersUtil
 				fi.Append(TEXT(" Profile"));
 				if (InLevelIDC)
 				{
-					fi.Append(FString::Printf(TEXT(", level %d.%d"), InLevelIDC/30, InLevelIDC%30));
+					fi.Append(FString::Printf(TEXT(", level %d.%d"), InLevelIDC/30, (InLevelIDC%30)/3));
 				}
 				else
 				{
@@ -1712,13 +1712,17 @@ namespace ElectraDecodersUtil
 								return false;
 							}
 						}
-						else if (sps.num_short_term_ref_pic_sets > 1)
+						else if (sps.num_short_term_ref_pic_sets)
 						{
 							const uint32 NumBits = FMath::CeilLogTwo(sps.num_short_term_ref_pic_sets);
 							OutSlice.short_term_ref_pic_set_idx = br.GetBits(NumBits);
 							RANGE_CHECK_FAILURE((int32)OutSlice.short_term_ref_pic_set_idx < sps.num_short_term_ref_pic_sets, TEXT("short_term_ref_pic_set_idx"));
 							OutSlice.st_ref_pic_set = sps.st_ref_pic_set[OutSlice.short_term_ref_pic_set_idx];
 						}
+						else
+						{
+							RANGE_CHECK_FAILURE(false, TEXT("sps.num_short_term_ref_pic_sets is 0"));
+						}
 						uint64 NumBitsRemainingAfter = br.GetRemainingBits();
 						OutSlice.NumBitsForShortTermRefs = NumBitsRemainingBefore - NumBitsRemainingAfter;
 
diff --git a/Engine/Plugins/Media/ElectraCodecs/Source/ElectraDecoders/Private/Windows/h265/H265_VideoDecoder_Windows.cpp b/Engine/Plugins/Media/ElectraCodecs/Source/ElectraDecoders/Private/Windows/h265/H265_VideoDecoder_Windows.cpp
--- a/Engine/Plugins/Media/ElectraCodecs/Source/ElectraDecoders/Private/Windows/h265/H265_VideoDecoder_Windows.cpp
+++ b/Engine/Plugins/Media/ElectraCodecs/Source/ElectraDecoders/Private/Windows/h265/H265_VideoDecoder_Windows.cpp
@@ -188,9 +188,9 @@ void IElectraVideoDecoderH265_DX::PlatformGetSupportedConfigurations(TArray<FSup
 		IElectraVideoDecoderH265_DX_Platform::DecoderConfigurations.Empty();
 
 		// Main
-		IElectraVideoDecoderH265_DX_Platform::DecoderConfigurations.Emplace(IElectraVideoDecoderH265_DX::FSupportedConfiguration(0, 1, 0, 153, 120, 4096, 2304, 0));
+		IElectraVideoDecoderH265_DX_Platform::DecoderConfigurations.Emplace(IElectraVideoDecoderH265_DX::FSupportedConfiguration(0, 1, 0, 156, 120, 4096, 2304, 0));
 		// Main10
-		IElectraVideoDecoderH265_DX_Platform::DecoderConfigurations.Emplace(IElectraVideoDecoderH265_DX::FSupportedConfiguration(0, 2, 0, 153, 120, 4096, 2304, 0));
+		IElectraVideoDecoderH265_DX_Platform::DecoderConfigurations.Emplace(IElectraVideoDecoderH265_DX::FSupportedConfiguration(0, 2, 0, 156, 120, 4096, 2304, 0));
 
 		IElectraVideoDecoderH265_DX_Platform::bDecoderConfigurationsDirty = false;
 	}
diff --git a/Engine/Plugins/Media/ElectraPlayer/Source/ElectraProtron/Private/Player/ElectraProtronPlayerDecoder.cpp b/Engine/Plugins/Media/ElectraPlayer/Source/ElectraProtron/Private/Player/ElectraProtronPlayerDecoder.cpp
--- a/Engine/Plugins/Media/ElectraPlayer/Source/ElectraProtron/Private/Player/ElectraProtronPlayerDecoder.cpp
+++ b/Engine/Plugins/Media/ElectraPlayer/Source/ElectraProtron/Private/Player/ElectraProtronPlayerDecoder.cpp
@@ -85,7 +85,8 @@ uint32 FElectraProtronPlayer::FImpl::FDecoderThread::Run()
 			PerformFlush();
 			FlushedSignal->Signal();
 		}
-		HandleOutputFrame();
+		while(HandleOutputFrame())
+		{}
 
 		// If an error occurred, do not do anything to make matters worse.
 		// Just wait for the player to be closed.
@@ -760,12 +761,13 @@ void FElectraProtronPlayer::FImpl::FDecoderThread::StepTrackIterator()
 }
 
 
-void FElectraProtronPlayer::FImpl::FDecoderThread::HandleOutputFrame()
+bool FElectraProtronPlayer::FImpl::FDecoderThread::HandleOutputFrame()
 {
 	if (!DecoderInstance.IsValid())
 	{
-		return;
+		return false;
 	}
+	bool bGetNext = false;
 
 	// If we still have unhandled output, do not get new output yet.
 	if (!CurrentDecoderOutput.IsValid())
@@ -779,7 +781,7 @@ void FElectraProtronPlayer::FImpl::FDecoderThread::HandleOutputFrame()
 			if (!Out.IsValid())
 			{
 				// No output although advertised?
-				return;
+				return false;
 			}
 			if (DecoderTypeIndex == CodecTypeIndex(Electra::FCodecTypeFormat::EType::Video) && Out->GetType() == IElectraDecoderOutput::EType::Video)
 			{
@@ -793,14 +795,14 @@ void FElectraProtronPlayer::FImpl::FDecoderThread::HandleOutputFrame()
 			if (Out.IsValid())
 			{
 				Out.Reset();
-				return;
+				return false;
 			}
 			// Find the matching input sample
 			if (!InDecoderInput.RemoveAndCopyValue(CurrentDecoderOutput->GetUserValue(), InputForCurrentDecoderOutput))
 			{
 				check(!"no input for this output?");
 				CurrentDecoderOutput.Reset();
-				return;
+				return false;
 			}
 		}
 		else if (OutputStatus == IElectraDecoder::EOutputStatus::EndOfData)
@@ -985,6 +987,7 @@ void FElectraProtronPlayer::FImpl::FDecoderThread::HandleOutputFrame()
 				}
 				CurrentDecoderOutput.Reset();
 				InputForCurrentDecoderOutput.Reset();
+				bGetNext = true;
 			}
 		}
 		else if (DecoderTypeIndex == CodecTypeIndex(Electra::FCodecTypeFormat::EType::Audio))
@@ -1100,6 +1103,7 @@ void FElectraProtronPlayer::FImpl::FDecoderThread::HandleOutputFrame()
 			}
 		}
 	}
+	return bGetNext;
 }
 
 
diff --git a/Engine/Plugins/Media/ElectraPlayer/Source/ElectraProtron/Private/Player/ElectraProtronPlayerImpl.h b/Engine/Plugins/Media/ElectraPlayer/Source/ElectraProtron/Private/Player/ElectraProtronPlayerImpl.h
--- a/Engine/Plugins/Media/ElectraPlayer/Source/ElectraProtron/Private/Player/ElectraProtronPlayerImpl.h
+++ b/Engine/Plugins/Media/ElectraPlayer/Source/ElectraProtron/Private/Player/ElectraProtronPlayerImpl.h
@@ -643,7 +643,7 @@ private:
 		bool CreateDecoder();
 		void DestroyDecoder();
 		void DecodeOneFrame();
-		void HandleOutputFrame();
+		bool HandleOutputFrame();
 		void FlushForEndOrLooping();
 		void PerformFlush();
 
